[Trace] Added TraceUtils Plugin and TraceUtils module.
Added TraceUtilLibrary to expose FTraceAuxiliary functions to blueprints.
Expose FChannel::EnumerateChannels to other modules

diff --git a/Engine/Plugins/TraceUtilities/Source/TraceUtilities/Private/TraceUtilLibrary.cpp b/Engine/Plugins/TraceUtilities/Source/TraceUtilities/Private/TraceUtilLibrary.cpp
new file mode 100644
--- /dev/null
+++ b/Engine/Plugins/TraceUtilities/Source/TraceUtilities/Private/TraceUtilLibrary.cpp
@@ -0,0 +1,85 @@
+#include "TraceUtilLibrary.h"
+
+#include "ProfilingDebugging/TraceAuxiliary.h"
+#include "Modules/ModuleManager.h"
+
+IMPLEMENT_MODULE(FDefaultModuleImpl, TraceUtilities)
+
+bool UTraceUtilLibrary::StartTraceToFile(const FString& FileName, const TArray<FString>& Channels)
+{
+	FString ChannelString = FString::Join(Channels, TEXT(","));
+	return FTraceAuxiliary::Start(FTraceAuxiliary::EConnectionType::File ,*FileName, *ChannelString);
+};
+
+bool UTraceUtilLibrary::StartTraceSendTo(const FString& Target, const TArray<FString>& Channels)
+{
+	FString ChannelString = FString::Join(Channels, TEXT(","));
+	return FTraceAuxiliary::Start(FTraceAuxiliary::EConnectionType::File ,*Target, *ChannelString);
+}
+
+bool UTraceUtilLibrary::StopTracing()
+{
+	return FTraceAuxiliary::Stop();
+}
+
+bool UTraceUtilLibrary::PauseTracing()
+{
+	return FTraceAuxiliary::Pause();
+}
+
+bool UTraceUtilLibrary::ResumeTracing()
+{
+	return FTraceAuxiliary::Resume();
+}
+
+bool UTraceUtilLibrary::IsTracing()
+{
+	return FTraceAuxiliary::IsConnected();
+}
+
+bool UTraceUtilLibrary::ToggleChannel(const FString& ChannelName, bool enabled)
+{
+	return UE::Trace::ToggleChannel(*ChannelName, enabled);
+}
+
+bool UTraceUtilLibrary::IsChannelEnabled(const FString& ChannelName)
+{
+	return UE::Trace::IsChannel(*ChannelName);
+}
+
+TArray<FString> UTraceUtilLibrary::GetEnabledChannels()
+{
+	TStringBuilder<256> StringBuilder;
+	FTraceAuxiliary::GetActiveChannelsString(StringBuilder);
+	TArray<FString> Out;
+	FString ChannelsString(StringBuilder.ToView());
+	ChannelsString.ParseIntoArray(Out, TEXT(","));
+	return Out;
+}
+
+TArray<FString> UTraceUtilLibrary::GetAllChannels()
+{
+	TArray<FString> Channels;
+	
+	UE::Trace::EnumerateChannels([](const ANSICHAR* Name, bool IsEnabled, void* Channels)
+	{
+		static_cast<TArray<FString>*>(Channels)->Add(ANSI_TO_TCHAR(Name));
+	}, &Channels);
+
+	return Channels;
+}
+
+void UTraceUtilLibrary::TraceBookmark(const FString& Name)
+{
+	TRACE_BOOKMARK(TEXT("%s"), *Name);
+}
+
+void UTraceUtilLibrary::TraceMarkRegionStart(const FString& Name)
+{
+	TRACE_BOOKMARK(TEXT("RegionStart:%s"), *Name);
+}
+
+void UTraceUtilLibrary::TraceMarkRegionEnd(const FString& Name)
+{
+	TRACE_BOOKMARK(TEXT("RegionEnd:%s"), *Name);
+}
diff --git a/Engine/Plugins/TraceUtilities/Source/TraceUtilities/Public/TraceUtilLibrary.h b/Engine/Plugins/TraceUtilities/Source/TraceUtilities/Public/TraceUtilLibrary.h
new file mode 100644
--- /dev/null
+++ b/Engine/Plugins/TraceUtilities/Source/TraceUtilities/Public/TraceUtilLibrary.h
@@ -0,0 +1,57 @@
+#pragma once 
+
+#include "Kismet/BlueprintFunctionLibrary.h"
+#include "TraceUtilLibrary.generated.h"
+
+UCLASS()
+class TRACEUTILITIES_API UTraceUtilLibrary : public UBlueprintFunctionLibrary
+{
+    GENERATED_BODY()
+    
+    public:
+    
+    UFUNCTION(BlueprintCallable, Category = "Perf | Insights Trace")
+    static bool StartTraceToFile(const FString& FileName, const TArray<FString>& Channels);
+    
+    UFUNCTION(BlueprintCallable, Category = "Perf | Insights Trace")
+    static bool StartTraceSendTo(const FString& Target, const TArray<FString>& Channels);
+    
+    UFUNCTION(BlueprintCallable, Category = "Perf | Insights Trace")
+    static bool StopTracing();
+    
+    UFUNCTION(BlueprintCallable, Category = "Perf | Insights Trace")
+    static bool PauseTracing();
+    
+    UFUNCTION(BlueprintCallable, Category = "Perf | Insights Trace")
+    static bool ResumeTracing();
+    
+    UFUNCTION(BlueprintCallable, Category = "Perf | Insights Trace")
+    static bool IsTracing();
+    
+    UFUNCTION(BlueprintCallable, Category = "Perf | Insights Trace")
+    static bool ToggleChannel(const FString& ChannelName, bool enabled);
+    
+    UFUNCTION(BlueprintCallable, Category = "Perf | Insights Trace")
+    static bool IsChannelEnabled(const FString& ChannelName);
+    
+    UFUNCTION(BlueprintCallable, Category = "Perf | Insights Trace")
+    static TArray<FString> GetEnabledChannels();
+    
+    UFUNCTION(BlueprintCallable, Category = "Perf | Insights Trace")
+    static TArray<FString> GetAllChannels();
+    
+    UFUNCTION(BlueprintCallable, Category = "Perf | Insights Trace")
+    static void TraceBookmark(const FString& Name);
+
+    /**
+     * Traces a "RegionStart:Name" bookmark, where Name is defined through the Name parameter.
+     */
+    UFUNCTION(BlueprintCallable, Category = "Perf | Insights Trace")
+    static void TraceMarkRegionStart(const FString& Name);
+ 
+	/**
+	 * Traces a "RegionEnd:Name" bookmark, where Name is defined through the Name parameter.
+	 */   
+    UFUNCTION(BlueprintCallable, Category = "Perf | Insights Trace")
+    static void TraceMarkRegionEnd(const FString& Name);
+};
\ No newline at end of file
diff --git a/Engine/Plugins/TraceUtilities/Source/TraceUtilities/TraceUtilities.Build.cs b/Engine/Plugins/TraceUtilities/Source/TraceUtilities/TraceUtilities.Build.cs
new file mode 100644
--- /dev/null
+++ b/Engine/Plugins/TraceUtilities/Source/TraceUtilities/TraceUtilities.Build.cs
@@ -0,0 +1,29 @@
+﻿// Copyright Epic Games, Inc. All Rights Reserved.
+
+using UnrealBuildTool;
+
+public class TraceUtilities : ModuleRules
+{
+	public TraceUtilities(ReadOnlyTargetRules Target) : base(Target)
+	{
+		PCHUsage = ModuleRules.PCHUsageMode.UseExplicitOrSharedPCHs;
+
+		PublicDependencyModuleNames.AddRange(
+			new string[]
+			{
+				"Core",
+				// ... add other public dependencies that you statically link with here ...
+			}
+			);
+
+		PrivateDependencyModuleNames.AddRange(
+			new string[]
+			{
+				"CoreUObject",
+				"Engine",
+				"TraceLog",
+				// ... add private dependencies that you statically link with here ...	
+			}
+			);
+	}
+}
diff --git a/Engine/Plugins/TraceUtilities/TraceUtilities.uplugin b/Engine/Plugins/TraceUtilities/TraceUtilities.uplugin
new file mode 100644
--- /dev/null
+++ b/Engine/Plugins/TraceUtilities/TraceUtilities.uplugin
@@ -0,0 +1,24 @@
+﻿{
+  "FileVersion": 3,
+  "Version": 1,
+  "VersionName": "1.0",
+  "FriendlyName": "TraceUtilities",
+  "Description": "A collection of tools that increase usability of Unreal Insights and the Trace Framework",
+  "Category": "Other",
+  "CreatedBy": "Epic Games",
+  "CreatedByURL": "",
+  "DocsURL": "",
+  "MarketplaceURL": "",
+  "CanContainContent": false,
+  "IsBetaVersion": false,
+  "IsExperimentalVersion": false,
+  "Installed": false,
+  "EnabledByDefault" : true,
+  "Modules": [
+    {
+      "Name": "TraceUtilities",
+      "Type": "Runtime",
+      "LoadingPhase": "Default"
+    }
+  ]
+}
\ No newline at end of file
