TraceUtilities Plugin: fix missing copyright notice

diff --git a/Engine/Plugins/TraceUtilities/Source/TraceUtilities/Private/TraceUtilLibrary.cpp b/Engine/Plugins/TraceUtilities/Source/TraceUtilities/Private/TraceUtilLibrary.cpp
--- a/Engine/Plugins/TraceUtilities/Source/TraceUtilities/Private/TraceUtilLibrary.cpp
+++ b/Engine/Plugins/TraceUtilities/Source/TraceUtilities/Private/TraceUtilLibrary.cpp
@@ -1,3 +1,5 @@
+// Copyright Epic Games, Inc. All Rights Reserved.
+
 #include "TraceUtilLibrary.h"
 
 #include "ProfilingDebugging/TraceAuxiliary.h"
diff --git a/Engine/Plugins/TraceUtilities/Source/TraceUtilities/Public/TraceUtilLibrary.h b/Engine/Plugins/TraceUtilities/Source/TraceUtilities/Public/TraceUtilLibrary.h
--- a/Engine/Plugins/TraceUtilities/Source/TraceUtilities/Public/TraceUtilLibrary.h
+++ b/Engine/Plugins/TraceUtilities/Source/TraceUtilities/Public/TraceUtilLibrary.h
@@ -1,3 +1,4 @@
+// Copyright Epic Games, Inc. All Rights Reserved.
 #pragma once 
 
 #include "Kismet/BlueprintFunctionLibrary.h"
