TraceUtilities Plugin:
- added a new module named EditorTraceUtilities
- Moved UnrealInsightsLauncher class and Editor Integration from TraceDataFiltering Plugin into EditorTraceUtilities
- Added an additional button to open the currently active trace from the Editor. Will open local .utrace files or try to open the live trace from the storebrowser.

diff --git a/Engine/Plugins/Developer/TraceDataFiltering/Source/UnrealInsightsInterface/Private/UnrealInsightsInterfaceModule.cpp b/Engine/Plugins/Developer/TraceDataFiltering/Source/UnrealInsightsInterface/Private/UnrealInsightsInterfaceModule.cpp
deleted file mode 100644
diff --git a/Engine/Plugins/Developer/TraceDataFiltering/Source/UnrealInsightsInterface/Private/UnrealInsightsLauncher.cpp b/Engine/Plugins/Developer/TraceDataFiltering/Source/UnrealInsightsInterface/Private/UnrealInsightsLauncher.cpp
deleted file mode 100644
diff --git a/Engine/Plugins/Developer/TraceDataFiltering/Source/UnrealInsightsInterface/Private/UnrealInsightsLauncher.h b/Engine/Plugins/Developer/TraceDataFiltering/Source/UnrealInsightsInterface/Private/UnrealInsightsLauncher.h
deleted file mode 100644
diff --git a/Engine/Plugins/Developer/TraceDataFiltering/Source/UnrealInsightsInterface/Public/UnrealInsightsInterfaceModule.h b/Engine/Plugins/Developer/TraceDataFiltering/Source/UnrealInsightsInterface/Public/UnrealInsightsInterfaceModule.h
deleted file mode 100644
diff --git a/Engine/Plugins/Developer/TraceDataFiltering/Source/UnrealInsightsInterface/UnrealInsightsInterface.Build.cs b/Engine/Plugins/Developer/TraceDataFiltering/Source/UnrealInsightsInterface/UnrealInsightsInterface.Build.cs
deleted file mode 100644
diff --git a/Engine/Plugins/Developer/TraceDataFiltering/TraceDataFilters.uplugin b/Engine/Plugins/Developer/TraceDataFiltering/TraceDataFilters.uplugin
--- a/Engine/Plugins/Developer/TraceDataFiltering/TraceDataFilters.uplugin
+++ b/Engine/Plugins/Developer/TraceDataFiltering/TraceDataFilters.uplugin
@@ -23,11 +23,6 @@
 			"ProgramAllowList": [
 				"UnrealInsights"
 			]
-		},
-		{
-			"Name": "UnrealInsightsInterface",
-			"Type": "Editor",
-			"LoadingPhase": "Default"
 		}
 	]
 }
\ No newline at end of file
diff --git a/Engine/Plugins/TraceUtilities/Source/EditorTraceUtilities/EditorTraceUtilities.Build.cs b/Engine/Plugins/TraceUtilities/Source/EditorTraceUtilities/EditorTraceUtilities.Build.cs
new file mode 100644
--- /dev/null
+++ b/Engine/Plugins/TraceUtilities/Source/EditorTraceUtilities/EditorTraceUtilities.Build.cs
@@ -0,0 +1,33 @@
+﻿using UnrealBuildTool;
+
+public class EditorTraceUtilities : ModuleRules
+{
+    public EditorTraceUtilities(ReadOnlyTargetRules Target) : base(Target)
+    {
+        PCHUsage = ModuleRules.PCHUsageMode.UseExplicitOrSharedPCHs;
+
+        PublicDependencyModuleNames.AddRange(
+            new string[]
+            {
+                "Core",
+            }
+        );
+
+        PrivateDependencyModuleNames.AddRange(
+            new string[]
+            {
+                "CoreUObject",
+                "Engine",
+                "Slate",
+                "SlateCore",
+                "ToolMenus",
+                "ToolWidgets",
+                "EditorFramework",
+                "ApplicationCore",
+                "UATHelper",
+                "TraceInsights",
+                "TraceAnalysis"
+            }
+        );
+    }
+}
\ No newline at end of file
diff --git a/Engine/Plugins/TraceUtilities/Source/EditorTraceUtilities/Private/EditorTraceUtilities.cpp b/Engine/Plugins/TraceUtilities/Source/EditorTraceUtilities/Private/EditorTraceUtilities.cpp
new file mode 100644
--- /dev/null
+++ b/Engine/Plugins/TraceUtilities/Source/EditorTraceUtilities/Private/EditorTraceUtilities.cpp
@@ -0,0 +1,27 @@
+﻿#include "EditorTraceUtilities.h"
+
+#include "ToolMenus.h"
+#include "UnrealInsightsLauncher.h"
+
+DEFINE_LOG_CATEGORY(LogTraceUtilities)
+
+#define LOCTEXT_NAMESPACE "FEditorTraceUtilitiesModule"
+
+void FEditorTraceUtilitiesModule::StartupModule()
+{
+	LLM_SCOPE_BYNAME(TEXT("Insights"));
+	RegisterStartupCallbackHandle = UToolMenus::RegisterStartupCallback(
+		FSimpleMulticastDelegate::FDelegate::CreateSP(
+			FUnrealInsightsLauncher::Get().ToSharedRef(),
+			&FUnrealInsightsLauncher::RegisterMenus));
+}
+
+void FEditorTraceUtilitiesModule::ShutdownModule()
+{
+	LLM_SCOPE_BYNAME(TEXT("Insights"));
+	UToolMenus::UnRegisterStartupCallback(RegisterStartupCallbackHandle);
+}
+
+#undef LOCTEXT_NAMESPACE
+    
+IMPLEMENT_MODULE(FEditorTraceUtilitiesModule, EditorTraceUtilities)
\ No newline at end of file
diff --git a/Engine/Plugins/TraceUtilities/Source/EditorTraceUtilities/Private/UnrealInsightsLauncher.cpp b/Engine/Plugins/TraceUtilities/Source/EditorTraceUtilities/Private/UnrealInsightsLauncher.cpp
new file mode 100644
--- /dev/null
+++ b/Engine/Plugins/TraceUtilities/Source/EditorTraceUtilities/Private/UnrealInsightsLauncher.cpp
@@ -0,0 +1,226 @@
+// Copyright Epic Games, Inc. All Rights Reserved.
+
+#include "UnrealInsightsLauncher.h"
+#include "EditorTraceUtilities.h"
+
+#include "Styling/AppStyle.h"
+#include "IUATHelperModule.h"
+#include "Logging/LogMacros.h"
+#include "Logging/MessageLog.h"
+#include "MessageLog/Public/MessageLogModule.h"
+#include "ToolMenus.h"
+#include "ProfilingDebugging/TraceAuxiliary.h"
+#include "Trace/StoreClient.h"
+
+#define LOCTEXT_NAMESPACE "FUnrealInsightsLauncher"
+
+
+TSharedPtr<FUnrealInsightsLauncher> FUnrealInsightsLauncher::Instance = nullptr;
+
+FUnrealInsightsLauncher::FUnrealInsightsLauncher()
+	: LogListingName(TEXT("UnrealInsights"))
+{
+
+}
+
+FUnrealInsightsLauncher::~FUnrealInsightsLauncher()
+{
+
+}
+	
+void FUnrealInsightsLauncher::RegisterMenus()
+{
+	FToolMenuOwnerScoped OwnerScoped(this);
+
+	UToolMenu* ProfileMenu = UToolMenus::Get()->ExtendMenu("MainFrame.MainMenu.Tools");
+	if (ProfileMenu)
+	{
+		FToolMenuSection& Section = ProfileMenu->AddSection("Unreal Insights", FText::FromString(TEXT("Unreal Insights")));
+		Section.AddMenuEntry("OpenUnrealInsights",
+			LOCTEXT("OpenUnrealInsights_Label", "Run Unreal Insights"),
+			LOCTEXT("OpenUnrealInsights_Desc", "Run the Unreal Insights standalone application."),
+			FSlateIcon(FAppStyle::GetAppStyleSetName(), "UnrealInsights.MenuIcon"),
+			FUIAction(FExecuteAction::CreateRaw(this, &FUnrealInsightsLauncher::RunUnrealInsights_Execute), FCanExecuteAction())
+		);
+		Section.AddMenuEntry("OpenLiveTrace",
+			LOCTEXT("OpenLiveTrace_Label", "Open active Trace"),
+			LOCTEXT("OpenLiveTrace_Desc", "Opens the currently running trace in Unreal Insights."),
+			FSlateIcon(FAppStyle::GetAppStyleSetName(), "UnrealInsights.MenuIcon"),
+			FUIAction(FExecuteAction::CreateRaw(this, &FUnrealInsightsLauncher::RunUnrealInsights_OpenLiveTrace), FCanExecuteAction::CreateLambda(
+				[](){ return FTraceAuxiliary::IsConnected();}))
+		);
+	}
+}
+
+FString FUnrealInsightsLauncher::GetInsightsApplicationPath()
+{
+	FString Path = FPlatformProcess::GenerateApplicationPath(TEXT("UnrealInsights"), EBuildConfiguration::Development);
+	return FPaths::ConvertRelativePathToFull(Path);
+}
+
+void FUnrealInsightsLauncher::RunUnrealInsights_Execute()
+{
+	FString Path = GetInsightsApplicationPath();
+	StartUnrealInsights(Path);
+}
+
+void FUnrealInsightsLauncher::RunUnrealInsights_OpenLiveTrace()
+{
+	TryOpenTraceFromDestination(FTraceAuxiliary::GetTraceDestination());
+}
+
+void FUnrealInsightsLauncher::StartUnrealInsights(const FString& Path, const FString& Parameters)
+{
+    if (!FPaths::FileExists(Path))
+    {
+    	TryBuildUnrealInsightsExe(Path, Parameters);
+    	return;
+    }
+	
+	constexpr bool bLaunchDetached = true;
+	constexpr bool bLaunchHidden = false;
+	constexpr bool bLaunchReallyHidden = false;
+
+	uint32 ProcessID = 0;
+	const int32 PriorityModifier = 0;
+	const TCHAR* OptionalWorkingDirectory = nullptr;
+
+	void* PipeWriteChild = nullptr;
+	void* PipeReadChild = nullptr;
+	FProcHandle Handle = FPlatformProcess::CreateProc(*Path, *Parameters, bLaunchDetached, bLaunchHidden, bLaunchReallyHidden, &ProcessID, PriorityModifier, OptionalWorkingDirectory, PipeWriteChild, PipeReadChild);
+
+	if (Handle.IsValid())
+	{
+		UE_LOG(LogTraceUtilities, Log, TEXT("Launched Unreal Insights executable: %s %s"), *Path, *Parameters);
+	}
+	else
+	{
+		const FText	MessageBoxTextFmt = LOCTEXT("ExecutableNotFound_TextFmt", "Could not start Unreal Insights executable at path: {0}");
+		const FText MessageBoxText = FText::Format(MessageBoxTextFmt, FText::FromString(Path));
+		LogMessage(MessageBoxText);
+	}
+}
+
+
+void FUnrealInsightsLauncher::TryBuildUnrealInsightsExe(const FString& Path, const FString& LaunchParameters)
+{
+#if !PLATFORM_MAC
+	UE_LOG(LogTraceUtilities, Log, TEXT("Could not find the Unreal Insights executable: %s. Attempting to build UnrealInsights."), *Path);
+
+	FString Arguments;
+#if PLATFORM_WINDOWS
+	FText PlatformName = LOCTEXT("PlatformName_Windows", "Windows");
+	Arguments = TEXT("BuildTarget -Target=UnrealInsights -Platform=Win64");
+#elif PLATFORM_MAC
+	FText PlatformName = LOCTEXT("PlatformName_Mac", "Mac");
+	Arguments = TEXT("BuildTarget -Target=UnrealInsights -Platform=Mac");
+#elif PLATFORM_LINUX
+	FText PlatformName = LOCTEXT("PlatformName_Linux", "Linux");
+	Arguments = TEXT("BuildTarget -Target=UnrealInsights -Platform=Linux");
+#endif
+
+	IUATHelperModule::Get().CreateUatTask(Arguments, PlatformName, LOCTEXT("BuildingUnrealInsights", "Building Unreal Insights"),
+		LOCTEXT("BuildUnrealInsightsTask", "Build Unreal Insights Task"), FAppStyle::GetBrush(TEXT("MainFrame.CookContent")), [this, Path, LaunchParameters](FString Result, double Time)
+		{
+			if (Result.Equals(TEXT("Completed")))
+			{
+				this->StartUnrealInsights(Path, LaunchParameters);
+			}
+		});
+#else
+	const FText	MessageBoxTextFmt = LOCTEXT("ExecutableNotFoundManualBuild_TextFmt", "Could not find Unreal Insights executable. Have you built Unreal Insights?");
+	const FText MessageBoxText = FText::Format(MessageBoxTextFmt, FText::FromString(Path));
+	LogMessage(MessageBoxText);	
+#endif
+}
+
+void FUnrealInsightsLauncher::LogMessage(const FText& Message) {
+	FMessageLogModule& MessageLogModule = FModuleManager::LoadModuleChecked<FMessageLogModule>("MessageLog");
+	if (!MessageLogModule.IsRegisteredLogListing(LogListingName))
+	{
+		MessageLogModule.RegisterLogListing(LogListingName, LOCTEXT("UnrealInsights", "Unreal Insights"));
+	}
+	
+	FMessageLog ReportMessageLog(LogListingName);
+	TSharedRef<FTokenizedMessage> TokenizedMessage = FTokenizedMessage::Create(EMessageSeverity::Error, Message);
+	ReportMessageLog.AddMessage(TokenizedMessage);
+	ReportMessageLog.Notify();
+}
+
+bool FUnrealInsightsLauncher::TryOpenTraceFromDestination(const FString& Destination)
+{
+	if (Destination.IsEmpty())
+		return false;
+	
+	if (FPaths::FileExists(Destination) && FPaths::GetExtension(Destination).Equals(TEXT("utrace")))
+	{
+		return OpenTraceFile(Destination);
+	}
+
+	// assume it's a ip address of the target server for now so we don't introduce dependencies on the socket system
+	//if (ISocketSubsystem::Get(PLATFORM_SOCKETSUBSYSTEM)->GetAddressFromString(Destination))
+	return OpenActiveTraceFromStore(Destination);
+}
+
+bool FUnrealInsightsLauncher::OpenTraceFile(const FString& FilePath)
+{
+	StartUnrealInsights(GetInsightsApplicationPath(), FilePath);
+	return true;
+}
+
+bool FUnrealInsightsLauncher::OpenRemoteTrace(const FString& TraceHostAddress, const uint16 TraceHostPort, uint32 TraceID)
+{
+	TStringBuilder<512> ParamsSB;
+	ParamsSB += TEXT("-Store=");
+	ParamsSB += TraceHostAddress;
+	ParamsSB += TEXT(":");
+	ParamsSB.Appendf(TEXT("%u"), TraceHostPort);
+	ParamsSB += TEXT(" -OpenTraceID=");
+	ParamsSB.Appendf(TEXT("%u"), TraceID);
+	
+	StartUnrealInsights(GetInsightsApplicationPath(), ParamsSB.ToString());
+	return true;
+}
+
+bool FUnrealInsightsLauncher::OpenActiveTraceFromStore(const FString& TraceHostAddress)
+{
+	UE::Trace::FStoreClient* StoreClient = UE::Trace::FStoreClient::Connect(*TraceHostAddress);
+	
+	if (!StoreClient)
+	{
+		const FText	MessageBoxTextFmt = LOCTEXT("StoreClientConnectFailed_TextFmt", "Could not connect to StoreClient at {0}");
+     	const FText MessageBoxText = FText::Format(MessageBoxTextFmt, FText::FromString(TraceHostAddress));
+     	LogMessage(MessageBoxText);
+		//No active connection to store client
+		return false;
+	}
+	 
+	int SessionCount = StoreClient->GetSessionCount();
+	if (!SessionCount)
+	{
+		const FText	MessageBoxTextFmt = LOCTEXT("StoreClientNoSession_TextFmt", "No active session found in StoreClient at {0}");
+     	const FText MessageBoxText = FText::Format(MessageBoxTextFmt, FText::FromString(TraceHostAddress));
+     	LogMessage(MessageBoxText);
+		return false;
+	}
+	// Get first active session
+	const UE::Trace::FStoreClient::FSessionInfo* SessionInfo = StoreClient->GetSessionInfo(0);
+
+	if (!SessionInfo)
+	{
+		// Failed to retrieve SessionInfo for active Session with Index 0
+		return false;
+	}
+	
+	uint32 TraceId = SessionInfo->GetTraceId();
+	
+	if (!OpenRemoteTrace(TraceHostAddress, StoreClient->GetStorePort(), TraceId))
+	{
+		//Failed to open Trace File %s
+		return false;
+	}
+
+	return true;
+}
+
+#undef LOCTEXT_NAMESPACE
diff --git a/Engine/Plugins/TraceUtilities/Source/EditorTraceUtilities/Private/UnrealInsightsLauncher.h b/Engine/Plugins/TraceUtilities/Source/EditorTraceUtilities/Private/UnrealInsightsLauncher.h
new file mode 100644
--- /dev/null
+++ b/Engine/Plugins/TraceUtilities/Source/EditorTraceUtilities/Private/UnrealInsightsLauncher.h
@@ -0,0 +1,91 @@
+// Copyright Epic Games, Inc. All Rights Reserved.
+
+#pragma once
+
+#include "CoreMinimal.h"
+
+class FUnrealInsightsLauncher : public TSharedFromThis<FUnrealInsightsLauncher>
+{
+public:
+	FUnrealInsightsLauncher();
+	~FUnrealInsightsLauncher();
+
+	static TSharedPtr<FUnrealInsightsLauncher> Instance;
+	
+	static const TSharedPtr<FUnrealInsightsLauncher>& Get()
+	{
+		if (!Instance)
+		{
+			Instance = MakeShared<FUnrealInsightsLauncher>();
+		}
+		return Instance;
+	}
+
+	
+	void RegisterMenus();
+
+	/**
+	 * Returns the full path to UnrealInsights.exe for the current engine installation 
+	 */
+	FString GetInsightsApplicationPath();
+
+	/**
+	 * Launches UnrealInsights.exe from the given Path, displays a editor message if it fails.
+	 * @param Path The full filename of UnrealInsights.exe to launch
+	 * @param Parameters The command line parameters to use when launching the exe
+	 */	
+	void StartUnrealInsights(const FString& Path, const FString& Parameters = TEXT(""));
+	
+	/**
+	 * Try to open a trace file from the trace target returned by UE::Trace::GetTraceDestination().
+	 * Destination may contain a file name or IP address of a trace server. In the latter case we try to query the
+	 * file name of the active/live trace through the storebrowser API and open the remote trace in Insights.
+	 * @param Destination The trace destination of the currently running trace session.
+	 */
+	bool TryOpenTraceFromDestination(const FString& Destination );
+
+	/**
+	 * Launches UnrealInsights.exe with FilePath as the only parameter
+	 * @param  FilePath Full file name and path to a .utrace file on the local machine
+	 * @return 
+	 */
+	bool OpenTraceFile(const FString& FilePath);
+
+	/**
+	 * Launches UnrealInsights.exe with the -Store and -OpenTraceID parameters
+	 * @param  TraceHostAddress hostname or ip address of the trace host
+	 * @param  TraceHostPort Port the Trace Host is running on
+	 * @param  TraceID ID of the Trace to open, can be obtained through the UE::Trace::FStoreBrowser interface
+	 * @return 
+	 */
+	bool OpenRemoteTrace(const FString& TraceHostAddress, const uint16 TraceHostPort, uint32 TraceID);
+
+	/**
+	* This function tries to query the file name of the active/live trace through the storebrowser API
+	 * @param Hostname hostname or ip of the target UnrealTraceServerX
+	 */
+	bool OpenActiveTraceFromStore(const FString& TraceHostAddress);
+
+private:
+	/*
+	 * Tries to build UnrealInsights via UAT, will launch with forwarded parameters if successful.
+	 * Assumes that the Insights Executable in path belongs to this engine.
+	 */
+	void TryBuildUnrealInsightsExe(const FString& Path, const FString& LaunchParameters = TEXT(""));
+
+	/// callback from Options Menu Entry "Start Unreal Insights"
+	void RunUnrealInsights_Execute();
+
+	/**
+	 * Opens the currently live trace if connected.
+	 */
+	 void RunUnrealInsights_OpenLiveTrace();
+
+	/// Logs an error message to the MessageLog window
+	void LogMessage(const FText& Message);
+
+private:
+	/** The name of the Unreal Insights log listing. */
+	FName LogListingName;
+};
+
diff --git a/Engine/Plugins/TraceUtilities/Source/EditorTraceUtilities/Public/EditorTraceUtilities.h b/Engine/Plugins/TraceUtilities/Source/EditorTraceUtilities/Public/EditorTraceUtilities.h
new file mode 100644
--- /dev/null
+++ b/Engine/Plugins/TraceUtilities/Source/EditorTraceUtilities/Public/EditorTraceUtilities.h
@@ -0,0 +1,16 @@
+﻿#pragma once
+
+#include "CoreMinimal.h"
+#include "Modules/ModuleManager.h"
+
+DECLARE_LOG_CATEGORY_EXTERN(LogTraceUtilities, Log, All)
+
+class FEditorTraceUtilitiesModule : public IModuleInterface
+{
+public:
+    virtual void StartupModule() override;
+    virtual void ShutdownModule() override;
+    
+private:
+    FDelegateHandle RegisterStartupCallbackHandle;
+};
diff --git a/Engine/Plugins/TraceUtilities/TraceUtilities.uplugin b/Engine/Plugins/TraceUtilities/TraceUtilities.uplugin
--- a/Engine/Plugins/TraceUtilities/TraceUtilities.uplugin
+++ b/Engine/Plugins/TraceUtilities/TraceUtilities.uplugin
@@ -19,6 +19,11 @@
       "Name": "TraceUtilities",
       "Type": "Runtime",
       "LoadingPhase": "Default"
+    },
+    {
+      "Name": "EditorTraceUtilities",
+      "Type": "EditorNoCommandlet",
+      "LoadingPhase": "Default"
     }
   ]
 }
\ No newline at end of file
