TraceUtilities Plugin: Fix missing copyright boilerplate in EditorTraceUtilities module.

diff --git a/Engine/Plugins/TraceUtilities/Source/EditorTraceUtilities/EditorTraceUtilities.Build.cs b/Engine/Plugins/TraceUtilities/Source/EditorTraceUtilities/EditorTraceUtilities.Build.cs
--- a/Engine/Plugins/TraceUtilities/Source/EditorTraceUtilities/EditorTraceUtilities.Build.cs
+++ b/Engine/Plugins/TraceUtilities/Source/EditorTraceUtilities/EditorTraceUtilities.Build.cs
@@ -1,4 +1,6 @@
-﻿using UnrealBuildTool;
+﻿// Copyright Epic Games, Inc. All Rights Reserved.
+
+using UnrealBuildTool;
 
 public class EditorTraceUtilities : ModuleRules
 {
diff --git a/Engine/Plugins/TraceUtilities/Source/EditorTraceUtilities/Private/EditorTraceUtilities.cpp b/Engine/Plugins/TraceUtilities/Source/EditorTraceUtilities/Private/EditorTraceUtilities.cpp
--- a/Engine/Plugins/TraceUtilities/Source/EditorTraceUtilities/Private/EditorTraceUtilities.cpp
+++ b/Engine/Plugins/TraceUtilities/Source/EditorTraceUtilities/Private/EditorTraceUtilities.cpp
@@ -1,4 +1,6 @@
-﻿#include "EditorTraceUtilities.h"
+﻿// Copyright Epic Games, Inc. All Rights Reserved.
+
+#include "EditorTraceUtilities.h"
 
 #include "ToolMenus.h"
 #include "UnrealInsightsLauncher.h"
diff --git a/Engine/Plugins/TraceUtilities/Source/EditorTraceUtilities/Public/EditorTraceUtilities.h b/Engine/Plugins/TraceUtilities/Source/EditorTraceUtilities/Public/EditorTraceUtilities.h
--- a/Engine/Plugins/TraceUtilities/Source/EditorTraceUtilities/Public/EditorTraceUtilities.h
+++ b/Engine/Plugins/TraceUtilities/Source/EditorTraceUtilities/Public/EditorTraceUtilities.h
@@ -1,4 +1,6 @@
-﻿#pragma once
+﻿// Copyright Epic Games, Inc. All Rights Reserved.
+
+#pragma once
 
 #include "CoreMinimal.h"
 #include "Modules/ModuleManager.h"
