Fixed compile issue after merging with latest UE5/Main.

diff --git a/Engine/Plugins/TraceUtilities/Source/EditorTraceUtilities/EditorTraceUtilities.Build.cs b/Engine/Plugins/TraceUtilities/Source/EditorTraceUtilities/EditorTraceUtilities.Build.cs
--- a/Engine/Plugins/TraceUtilities/Source/EditorTraceUtilities/EditorTraceUtilities.Build.cs
+++ b/Engine/Plugins/TraceUtilities/Source/EditorTraceUtilities/EditorTraceUtilities.Build.cs
@@ -15,7 +15,14 @@ public class EditorTraceUtilities : ModuleRules
             }
         );
 
-        PrivateDependencyModuleNames.AddRange(
+		PrivateIncludePathModuleNames.AddRange(
+			new string[]
+			{
+				"MessageLog",
+			}
+		);
+
+		PrivateDependencyModuleNames.AddRange(
             new string[]
             {
                 "CoreUObject",
