Add support on Mac for building Insights when lanching from Editor if it is not already built.

diff --git a/Engine/Plugins/TraceUtilities/Source/EditorTraceUtilities/Private/UnrealInsightsLauncher.cpp b/Engine/Plugins/TraceUtilities/Source/EditorTraceUtilities/Private/UnrealInsightsLauncher.cpp
--- a/Engine/Plugins/TraceUtilities/Source/EditorTraceUtilities/Private/UnrealInsightsLauncher.cpp
+++ b/Engine/Plugins/TraceUtilities/Source/EditorTraceUtilities/Private/UnrealInsightsLauncher.cpp
@@ -17,6 +17,27 @@
 
 TSharedPtr<FUnrealInsightsLauncher> FUnrealInsightsLauncher::Instance = nullptr;
 
+// We use this Task to launch notifications from the Game Thread because on Mac the app is closed if the notification API is called from another thread.
+class FLogMessageOnGameThreadTask
+{
+public:
+	FLogMessageOnGameThreadTask(const FText& InMessage)
+		: Message(InMessage)
+	{}
+
+	FORCEINLINE TStatId GetStatId() const { RETURN_QUICK_DECLARE_CYCLE_STAT(FLogMessageOnGameThreadTask, STATGROUP_TaskGraphTasks); }
+	ENamedThreads::Type GetDesiredThread() { return ENamedThreads::Type::GameThread; }
+	static ESubsequentsMode::Type GetSubsequentsMode() { return ESubsequentsMode::TrackSubsequents; }
+
+	void DoTask(ENamedThreads::Type CurrentThread, const FGraphEventRef& MyCompletionGraphEvent)
+	{
+		FUnrealInsightsLauncher::Get()->LogMessage(Message);
+	}
+
+private:
+	FText Message;
+};
+
 FUnrealInsightsLauncher::FUnrealInsightsLauncher()
 	: LogListingName(TEXT("UnrealInsights"))
 {
@@ -83,6 +104,7 @@ void FUnrealInsightsLauncher::StartUnrealInsights(const FString& Path, const FSt
 
 	uint32 ProcessID = 0;
 	const int32 PriorityModifier = 0;
+
 	const TCHAR* OptionalWorkingDirectory = nullptr;
 
 	void* PipeWriteChild = nullptr;
@@ -97,14 +119,13 @@ void FUnrealInsightsLauncher::StartUnrealInsights(const FString& Path, const FSt
 	{
 		const FText	MessageBoxTextFmt = LOCTEXT("ExecutableNotFound_TextFmt", "Could not start Unreal Insights executable at path: {0}");
 		const FText MessageBoxText = FText::Format(MessageBoxTextFmt, FText::FromString(Path));
-		LogMessage(MessageBoxText);
+		LogMessageOnGameThread(MessageBoxText);
 	}
 }
 
 
 void FUnrealInsightsLauncher::TryBuildUnrealInsightsExe(const FString& Path, const FString& LaunchParameters)
 {
-#if !PLATFORM_MAC
 	UE_LOG(LogTraceUtilities, Log, TEXT("Could not find the Unreal Insights executable: %s. Attempting to build UnrealInsights."), *Path);
 
 	FString Arguments;
@@ -124,17 +145,19 @@ void FUnrealInsightsLauncher::TryBuildUnrealInsightsExe(const FString& Path, con
 		{
 			if (Result.Equals(TEXT("Completed")))
 			{
+#if PLATFORM_MAC
+				// On Mac we genereate the path again so that it includes the newly built executable.
+				FString NewPath = GetInsightsApplicationPath();
+				this->StartUnrealInsights(NewPath, LaunchParameters);
+#else
 				this->StartUnrealInsights(Path, LaunchParameters);
+#endif
 			}
 		});
-#else
-	const FText	MessageBoxTextFmt = LOCTEXT("ExecutableNotFoundManualBuild_TextFmt", "Could not find Unreal Insights executable. Have you built Unreal Insights?");
-	const FText MessageBoxText = FText::Format(MessageBoxTextFmt, FText::FromString(Path));
-	LogMessage(MessageBoxText);	
-#endif
 }
 
-void FUnrealInsightsLauncher::LogMessage(const FText& Message) {
+void FUnrealInsightsLauncher::LogMessage(const FText& Message)
+{
 	FMessageLogModule& MessageLogModule = FModuleManager::LoadModuleChecked<FMessageLogModule>("MessageLog");
 	if (!MessageLogModule.IsRegisteredLogListing(LogListingName))
 	{
@@ -147,6 +170,11 @@ void FUnrealInsightsLauncher::LogMessage(const FText& Message) {
 	ReportMessageLog.Notify();
 }
 
+void FUnrealInsightsLauncher::LogMessageOnGameThread(const FText& Message)
+{
+	TGraphTask<FLogMessageOnGameThreadTask>::CreateTask().ConstructAndDispatchWhenReady(Message);
+}
+
 bool FUnrealInsightsLauncher::TryOpenTraceFromDestination(const FString& Destination)
 {
 	if (Destination.IsEmpty())
diff --git a/Engine/Plugins/TraceUtilities/Source/EditorTraceUtilities/Private/UnrealInsightsLauncher.h b/Engine/Plugins/TraceUtilities/Source/EditorTraceUtilities/Private/UnrealInsightsLauncher.h
--- a/Engine/Plugins/TraceUtilities/Source/EditorTraceUtilities/Private/UnrealInsightsLauncher.h
+++ b/Engine/Plugins/TraceUtilities/Source/EditorTraceUtilities/Private/UnrealInsightsLauncher.h
@@ -6,6 +6,8 @@
 
 class FUnrealInsightsLauncher : public TSharedFromThis<FUnrealInsightsLauncher>
 {
+	friend class FLogMessageOnGameThreadTask;
+
 public:
 	FUnrealInsightsLauncher();
 	~FUnrealInsightsLauncher();
@@ -84,6 +86,9 @@ private:
 	/// Logs an error message to the MessageLog window
 	void LogMessage(const FText& Message);
 
+	/// Logs an error message to the MessageLog window using a task that is ran on the Game Thread.
+	void LogMessageOnGameThread(const FText& Message);
+
 private:
 	/** The name of the Unreal Insights log listing. */
 	FName LogListingName;
