Initial functionality for the Insights Status Bar Widget

diff --git a/Engine/Plugins/TraceUtilities/Source/EditorTraceUtilities/Private/EditorTraceUtilities.cpp b/Engine/Plugins/TraceUtilities/Source/EditorTraceUtilities/Private/EditorTraceUtilities.cpp
--- a/Engine/Plugins/TraceUtilities/Source/EditorTraceUtilities/Private/EditorTraceUtilities.cpp
+++ b/Engine/Plugins/TraceUtilities/Source/EditorTraceUtilities/Private/EditorTraceUtilities.cpp
@@ -9,6 +9,11 @@ DEFINE_LOG_CATEGORY(LogTraceUtilities)
 
 #define LOCTEXT_NAMESPACE "FEditorTraceUtilitiesModule"
 
+/**
+  * This function will add the SInsightsStatusBarWidget to the Editor's status bar at the bottom ("LevelEditor.StatusBar.ToolBar").
+  */
+void RegisterInsightsStatusWidgetWithToolMenu();
+
 void FEditorTraceUtilitiesModule::StartupModule()
 {
 	LLM_SCOPE_BYNAME(TEXT("Insights"));
@@ -16,6 +21,8 @@ void FEditorTraceUtilitiesModule::StartupModule()
 		FSimpleMulticastDelegate::FDelegate::CreateSP(
 			FUnrealInsightsLauncher::Get().ToSharedRef(),
 			&FUnrealInsightsLauncher::RegisterMenus));
+
+	RegisterInsightsStatusWidgetWithToolMenu();
 }
 
 void FEditorTraceUtilitiesModule::ShutdownModule()
diff --git a/Engine/Plugins/TraceUtilities/Source/EditorTraceUtilities/Private/InsightsStatusBar/EditorTraceUtilitiesStyle.h b/Engine/Plugins/TraceUtilities/Source/EditorTraceUtilities/Private/InsightsStatusBar/EditorTraceUtilitiesStyle.h
new file mode 100644
--- /dev/null
+++ b/Engine/Plugins/TraceUtilities/Source/EditorTraceUtilities/Private/InsightsStatusBar/EditorTraceUtilitiesStyle.h
@@ -0,0 +1,50 @@
+// Copyright Epic Games, Inc. All Rights Reserved.
+
+#pragma once
+
+#include "Styling/SlateStyle.h"
+#include "Styling/SlateStyleRegistry.h"
+#include "Styling/SlateStyleMacros.h"
+#include "Styling/StyleColors.h"
+
+class FEditorTraceUtilitiesStyle final
+	: public FSlateStyleSet
+{
+public:
+	FEditorTraceUtilitiesStyle()
+		: FSlateStyleSet("EditorTraceUtilitiesStyle")
+	{
+		const FVector2D Icon10x10(10.0f, 10.0f);
+		const FVector2D Icon16x16(16.0f, 16.0f);
+		const FVector2D Icon20x20(20.0f, 20.0f);
+		const FVector2D Icon24x24(24.0f, 24.0f);
+		const FVector2D Icon32x32(32.0f, 32.0f);
+		const FVector2D Icon40x40(40.0f, 40.0f);
+
+		SetContentRoot(FPaths::EnginePluginsDir() / TEXT("/TraceUtilities/Content/Slate"));
+		SetCoreContentRoot(FPaths::EngineContentDir() / TEXT("Slate"));
+
+		Set("Icons.OpenLiveSession", new CORE_IMAGE_BRUSH_SVG("Starship/Insights/Session", Icon16x16));
+
+		Set("Icons.Trace", new IMAGE_BRUSH_SVG("Trace", Icon16x16));
+		Set("Icons.RecordTrace", new IMAGE_BRUSH_SVG("RecordTrace", Icon16x16));
+		Set("Icons.RecordTraceOutline", new IMAGE_BRUSH_SVG("RecordTraceOutline", Icon16x16));
+		Set("Icons.RecordTraceWithOutline", new IMAGE_BRUSH_SVG("RecordTraceWithOutline", Icon16x16));
+		Set("Icons.RecordTraceStop", new IMAGE_BRUSH_SVG("RecordTraceStop", Icon16x16, FStyleColors::Error));
+		Set("Icons.StartTrace", new IMAGE_BRUSH_SVG("StartTrace", Icon16x16));
+		Set("Icons.TraceSnapshot", new IMAGE_BRUSH_SVG("TraceSnapshot", Icon16x16));
+
+		FSlateStyleRegistry::RegisterSlateStyle(*this);
+	}
+
+	static FEditorTraceUtilitiesStyle& Get()
+	{
+		static FEditorTraceUtilitiesStyle Inst;
+		return Inst;
+	}
+	
+	~FEditorTraceUtilitiesStyle()
+	{
+		FSlateStyleRegistry::UnRegisterSlateStyle(*this);
+	}
+};
diff --git a/Engine/Plugins/TraceUtilities/Source/EditorTraceUtilities/Private/InsightsStatusBar/SInsightsStatusBar.cpp b/Engine/Plugins/TraceUtilities/Source/EditorTraceUtilities/Private/InsightsStatusBar/SInsightsStatusBar.cpp
new file mode 100644
--- /dev/null
+++ b/Engine/Plugins/TraceUtilities/Source/EditorTraceUtilities/Private/InsightsStatusBar/SInsightsStatusBar.cpp
@@ -0,0 +1,539 @@
+// Copyright Epic Games, Inc. All Rights Reserved.
+
+#include "SInsightsStatusBar.h"
+
+#include "EditorTraceUtilitiesStyle.h"
+#include "Framework/Application/SlateApplication.h"
+#include "Framework/MultiBox/MultiBoxBuilder.h"
+#include "Framework/Notifications/NotificationManager.h"
+#include "Math/Color.h"
+#include "Modules/ModuleManager.h"
+#include "ProfilingDebugging/TraceAuxiliary.h"
+#include "Styling/StyleColors.h"
+#include "ToolMenus.h"
+#include "Trace/StoreClient.h"
+#include "UnrealInsightsLauncher.h"
+#include "Widgets/SBoxPanel.h"
+#include "Widgets/SOverlay.h"
+#include "Widgets/Images/SImage.h"
+#include "Widgets/Input/SButton.h"
+#include "Widgets/Input/SComboButton.h"
+#include "Widgets/Notifications/SNotificationList.h"
+
+#define LOCTEXT_NAMESPACE "InsightsEditor"
+
+const TCHAR* SInsightsStatusBarWidget::DefaultPreset = TEXT("default");
+const TCHAR* SInsightsStatusBarWidget::MemoryPreset= TEXT("default,memory");
+const TCHAR* SInsightsStatusBarWidget::TaskGraphPreset = TEXT("default,task");
+const TCHAR* SInsightsStatusBarWidget::ContextSwitchesPreset = TEXT("default,contextswitches");
+
+TSharedRef<SWidget> CreateInsightsStatusBarWidget()
+{
+	return SNew(SInsightsStatusBarWidget);
+}
+
+void RegisterInsightsStatusWidgetWithToolMenu()
+{
+	UToolMenu* Menu = UToolMenus::Get()->ExtendMenu(TEXT("LevelEditor.StatusBar.ToolBar"));
+
+	FToolMenuSection& InsightsSection = Menu->AddSection(TEXT("Insights"), FText::GetEmpty(), FToolMenuInsert(TEXT("Compile"), EToolMenuInsertType::Before));
+
+	InsightsSection.AddEntry(
+		FToolMenuEntry::InitWidget(TEXT("InsightsStatusBar"), CreateInsightsStatusBarWidget(), FText::GetEmpty(), true, false)
+	);
+}
+
+void SInsightsStatusBarWidget::Construct(const FArguments& InArgs)
+{
+	this->ChildSlot
+	[
+		SNew(SHorizontalBox)
+		+ SHorizontalBox::Slot()
+		.AutoWidth()
+		[
+			SNew(SComboButton)
+			.ContentPadding(FMargin(6.0f, 0.0f))
+			.MenuPlacement(MenuPlacement_AboveAnchor)
+			.ComboButtonStyle(&FAppStyle::Get().GetWidgetStyle<FComboButtonStyle>("StatusBar.StatusBarComboButton"))
+			.OnGetMenuContent(this, &SInsightsStatusBarWidget::MakeTraceMenu)
+			.HasDownArrow(true)
+			.ToolTipText(this, &SInsightsStatusBarWidget::GetTitleToolTipText)
+			.ButtonContent()
+			[
+				SNew(SHorizontalBox)
+
+				+ SHorizontalBox::Slot()
+				.Padding(0, 0, 3, 0)
+				.AutoWidth()
+				[
+					SNew(SImage)
+					.ColorAndOpacity(FSlateColor::UseForeground())
+					.Image(FEditorTraceUtilitiesStyle::Get().GetBrush("Icons.Trace"))
+				]
+
+				+ SHorizontalBox::Slot()
+				.VAlign(VAlign_Center)
+				.AutoWidth()
+				[
+					SNew(STextBlock)
+					.Text(LOCTEXT("Trace", "Trace"))
+					.TextStyle(&FAppStyle::Get().GetWidgetStyle<FTextBlockStyle>("DialogButtonText"))
+				]
+			]
+		]
+
+		+ SHorizontalBox::Slot()
+		.AutoWidth()
+		[
+			SNew(SButton)
+			.ButtonStyle(FAppStyle::Get(), "SimpleButton")
+			.ContentPadding(FMargin(0.0f, 0.0f, 0.0f, 3.0f))
+			.HAlign(HAlign_Left)
+			.VAlign(VAlign_Bottom)
+			.ToolTipText(this, &SInsightsStatusBarWidget::GetRecordingButtonTooltipText)
+			.OnClicked_Lambda([this]() { this->ToggleTracing_OnClicked(); return FReply::Handled(); })
+			.OnHovered_Lambda([this]() { this->bIsTraceRecordButtonHovered = true; })
+			.OnUnhovered_Lambda([this]() { this->bIsTraceRecordButtonHovered = false; })
+			.Content()
+			[
+				SNew(SOverlay)
+
+				+ SOverlay::Slot()
+				[
+					SNew(SImage)
+					.ColorAndOpacity(this, &SInsightsStatusBarWidget::GetRecordingButtonColor)
+					.Image(FEditorTraceUtilitiesStyle::Get().GetBrush("Icons.RecordTrace"))
+					.Visibility(this, &SInsightsStatusBarWidget::GetStartTraceIconVisibility)
+				]
+
+				+ SOverlay::Slot()
+				[
+					SNew(SImage)
+					.ColorAndOpacity(this, &SInsightsStatusBarWidget::GetRecordingButtonOutlineColor)
+					.Image(FEditorTraceUtilitiesStyle::Get().GetBrush("Icons.RecordTraceOutline"))
+					.Visibility(this, &SInsightsStatusBarWidget::GetStartTraceIconVisibility)
+				]
+
+				+ SOverlay::Slot()
+				[
+					SNew(SImage)
+					.Image(FEditorTraceUtilitiesStyle::Get().GetBrush("Icons.RecordTraceStop"))
+					.Visibility(this, &SInsightsStatusBarWidget::GetStopTraceIconVisibility)
+				]
+			]
+		]
+
+		+ SHorizontalBox::Slot()
+		.AutoWidth()
+		[
+			SNew(SButton)
+			.ButtonStyle(FAppStyle::Get(), "SimpleButton")
+			.ContentPadding(FMargin(4.0f, 0.0f, 0.0f, 3.0f))
+			.HAlign(HAlign_Left)
+			.VAlign(VAlign_Bottom)
+			.OnClicked_Lambda([this]() { SaveSnapshot(); return FReply::Handled(); })
+			.Content()
+			[
+				SNew(SImage)
+				.DesiredSizeOverride(CoreStyleConstants::Icon16x16)
+				.Image(FEditorTraceUtilitiesStyle::Get().GetBrush("Icons.TraceSnapshot"))
+				.ToolTipText(LOCTEXT("SaveSnapShot","Snapshot: Save Current Trace Buffer To File"))
+			]
+		]
+	];
+
+	Channels = DefaultPreset;
+}
+
+TSharedRef<SWidget> SInsightsStatusBarWidget::MakeTraceMenu()
+{
+	FMenuBuilder MenuBuilder(true, nullptr);
+
+	// To be activated
+	//MenuBuilder.BeginSection("TraceDataFiltering", LOCTEXT("TraceMenu_Section_DataFiltering", "Trace Data Filtering"));
+	//{
+	//	MenuBuilder.AddSubMenu
+	//	(
+	//		LOCTEXT("Channels", "Channels"),
+	//		LOCTEXT("Channels_Desc", "Select what trace channels to enable when tracing."),
+	//		FNewMenuDelegate::CreateSP(this, &SInsightsStatusBarWidget::Channels_BuildMenu),
+	//		false,
+	//		FSlateIcon(FEditorTraceUtilitiesStyle::Get().GetStyleSetName(), ("Icons.Trace"))
+	//	);
+	//}
+	//MenuBuilder.EndSection();
+
+	MenuBuilder.BeginSection("TraceDestination", LOCTEXT("TraceMenu_Section_Destination", "Trace Destination"));
+	{
+		MenuBuilder.AddMenuEntry(
+			LOCTEXT("ServerLabel", "Trace Store"),
+			LOCTEXT("FileLabelDesc", "Set the trace store as the trace destination."),
+			FSlateIcon(),
+			FUIAction(FExecuteAction::CreateSP(this, &SInsightsStatusBarWidget::SetTraceDestination_Execute, ETraceDestination::TraceStore),
+					  FCanExecuteAction::CreateSP(this, &SInsightsStatusBarWidget::SetTraceDestination_CanExecute),
+					  FIsActionChecked::CreateSP(this, &SInsightsStatusBarWidget::SetTraceDestination_IsChecked, ETraceDestination::TraceStore)),
+			NAME_None,
+			EUserInterfaceActionType::RadioButton
+		);
+
+		MenuBuilder.AddMenuEntry(
+			LOCTEXT("ServerLabel", "File"),
+			LOCTEXT("FileLabelDesc", "Set file as the trace destination."),
+			FSlateIcon(),
+			FUIAction(FExecuteAction::CreateSP(this, &SInsightsStatusBarWidget::SetTraceDestination_Execute, ETraceDestination::File),
+					  FCanExecuteAction::CreateSP(this, &SInsightsStatusBarWidget::SetTraceDestination_CanExecute),
+					  FIsActionChecked::CreateSP(this, &SInsightsStatusBarWidget::SetTraceDestination_IsChecked, ETraceDestination::File)),
+			NAME_None,
+			EUserInterfaceActionType::RadioButton
+		);
+	}
+	MenuBuilder.EndSection();
+
+	MenuBuilder.BeginSection("Tracing", LOCTEXT("TraceMenu_Section_Tracing", "Tracing"));
+	{
+		MenuBuilder.AddMenuEntry(
+			LOCTEXT("SaveSnapshotLabel", "Save Trace Snapshot"),
+			LOCTEXT("OpenInsightsTooltip", "Save the current trace buffer to file."),
+			FSlateIcon(FEditorTraceUtilitiesStyle::Get().GetStyleSetName(), "Icons.TraceSnapshot"),
+			FUIAction(FExecuteAction::CreateSP(this, &SInsightsStatusBarWidget::SaveSnapshot),
+					  FCanExecuteAction::CreateSP(this, &SInsightsStatusBarWidget::SaveSnapshot_CanExecute)),
+			NAME_None,
+			EUserInterfaceActionType::Button
+		);
+
+		MenuBuilder.AddMenuEntry(
+			TAttribute<FText>::CreateSP(this, &SInsightsStatusBarWidget::GetTraceMenuItemText),
+			TAttribute<FText>::CreateSP(this, &SInsightsStatusBarWidget::GetTraceMenuItemTooltipText),
+			FSlateIcon(FEditorTraceUtilitiesStyle::Get().GetStyleSetName(), "Icons.StartTrace"),
+			FUIAction(FExecuteAction::CreateSP(this, &SInsightsStatusBarWidget::ToggleTracing_OnClicked)),
+			NAME_None,
+			EUserInterfaceActionType::Button
+		);
+	}
+	MenuBuilder.EndSection();
+
+	MenuBuilder.BeginSection("Insights", LOCTEXT("TraceMenu_Section_Insights", "Insights"));
+	{
+		MenuBuilder.AddMenuEntry(
+			LOCTEXT("OpenInsightsLabel", "Unreal Insights (Session Browser)"),
+			LOCTEXT("OpenInsightsTooltip", "Launch the Unreal Insights Session Browser."),
+			FSlateIcon(FAppStyle::Get().GetStyleSetName(), "UnrealInsights.MenuIcon"),
+			FUIAction(FExecuteAction::CreateSP(this, &SInsightsStatusBarWidget::LaunchUnrealInsights_OnClicked)),
+			NAME_None,
+			EUserInterfaceActionType::Button
+		);
+
+		MenuBuilder.AddMenuEntry(
+			LOCTEXT("OpenLiveSessionLabel", "Open Live Session"),
+			LOCTEXT("OpenLiveSessionTooltip", "Opening the live session is possible only while tracing to the trace store."),
+			FSlateIcon(FEditorTraceUtilitiesStyle::Get().GetStyleSetName(), "Icons.OpenLiveSession"),
+			FUIAction(FExecuteAction::CreateSP(this, &SInsightsStatusBarWidget::OpenLiveSession_OnClicked),
+					  FCanExecuteAction::CreateLambda([this]() { return FTraceAuxiliary::IsConnected() && this->TraceDestination == ETraceDestination::TraceStore; })),
+			NAME_None,
+			EUserInterfaceActionType::Button
+		);
+	}
+	MenuBuilder.EndSection();
+
+
+	MenuBuilder.BeginSection("Locations", LOCTEXT("TraceMenu_Section_Locations", "Locations"));
+	{
+		MenuBuilder.AddMenuEntry(
+			LOCTEXT("OpenProfilingDirectoryLabel", "Open Profiling Directory"),
+			LOCTEXT("OpenProfilingDirectoryTooltip", "Opens the profiling directory of the current project. This is the location where traces to file are stored."),
+			FSlateIcon(FAppStyle::Get().GetStyleSetName(), "Icons.FolderOpen"),
+			FUIAction(FExecuteAction::CreateSP(this, &SInsightsStatusBarWidget::OpenProfilingDirectory_OnClicked)),
+			NAME_None,
+			EUserInterfaceActionType::Button
+		);
+
+		MenuBuilder.AddMenuEntry(
+			LOCTEXT("OpenTraceStoreLabel", "Open Trace Store Directory"),
+			LOCTEXT("OpenTraceStoreTooltip", "Open Trace Store Directory. This is the location where traces saved to the trace server are stored."),
+			FSlateIcon(FAppStyle::Get().GetStyleSetName(), "Icons.FolderOpen"),
+			FUIAction(FExecuteAction::CreateSP(this, &SInsightsStatusBarWidget::OpenTraceStoreDirectory_OnClicked)),
+			NAME_None,
+			EUserInterfaceActionType::Button
+		);
+	}
+	MenuBuilder.EndSection();
+
+	return MenuBuilder.MakeWidget();
+}
+
+void SInsightsStatusBarWidget::Channels_BuildMenu(FMenuBuilder& MenuBuilder)
+{
+	MenuBuilder.BeginSection("Presets", LOCTEXT("Presets", "Presets"));
+	{
+		MenuBuilder.AddMenuEntry(
+			LOCTEXT("DefautLabel", "Default"),
+			LOCTEXT("DefaultLabelDesc", "Activate the cpu,gpu,frame,log,bookmark channels."),
+			FSlateIcon(),
+			FUIAction(FExecuteAction::CreateSP(this, &SInsightsStatusBarWidget::SetTraceChannels, DefaultPreset),
+				FCanExecuteAction::CreateLambda([]() {return true; }),
+				FIsActionChecked::CreateSP(this, &SInsightsStatusBarWidget::IsPresetSet, DefaultPreset)),
+			NAME_None,
+			EUserInterfaceActionType::RadioButton
+		);
+
+		MenuBuilder.AddMenuEntry(
+			LOCTEXT("MemoryLabel", "Memory"),
+			LOCTEXT("MemoryLabelDesc", "Activate the memtag,memalloc,callstack,module and the default channels."),
+			FSlateIcon(),
+			FUIAction(FExecuteAction::CreateSP(this, &SInsightsStatusBarWidget::SetTraceChannels, MemoryPreset),
+				FCanExecuteAction::CreateLambda([]() {return false; }),
+				FIsActionChecked::CreateSP(this, &SInsightsStatusBarWidget::IsPresetSet, MemoryPreset)),
+			NAME_None,
+			EUserInterfaceActionType::RadioButton
+		);
+
+		MenuBuilder.AddMenuEntry(
+			LOCTEXT("TaskGraphLabel", "Task Graph"),
+			LOCTEXT("TaskGraphLabelDesc", "Activate the task and the default channels."),
+			FSlateIcon(),
+			FUIAction(FExecuteAction::CreateSP(this, &SInsightsStatusBarWidget::SetTraceChannels, TaskGraphPreset),
+				FCanExecuteAction::CreateLambda([]() {return false; }),
+				FIsActionChecked::CreateSP(this, &SInsightsStatusBarWidget::IsPresetSet, TaskGraphPreset)),
+			NAME_None,
+			EUserInterfaceActionType::RadioButton
+		);
+
+		MenuBuilder.AddMenuEntry(
+			LOCTEXT("Context Switches Label", "Context Switches"),
+			LOCTEXT("Context Switches LabelDesc", "Activate the contextswitches and the default channels."),
+			FSlateIcon(),
+			FUIAction(FExecuteAction::CreateSP(this, &SInsightsStatusBarWidget::SetTraceChannels, ContextSwitchesPreset),
+				FCanExecuteAction::CreateLambda([]() { return false; }),
+				FIsActionChecked::CreateSP(this, &SInsightsStatusBarWidget::IsPresetSet, ContextSwitchesPreset)),
+			NAME_None,
+			EUserInterfaceActionType::RadioButton
+		);
+	}
+	MenuBuilder.EndSection();
+}
+
+FText SInsightsStatusBarWidget::GetTitleToolTipText() const
+{
+	FTextBuilder DescBuilder;
+
+	const TCHAR* Dest = FTraceAuxiliary::GetTraceDestination();
+	
+	if (*Dest != 0)
+	{
+		DescBuilder.AppendLineFormat(LOCTEXT("TracingToText", "Tracing to: {0}"), FText::FromString(Dest));
+	}
+	else
+	{
+		DescBuilder.AppendLine(LOCTEXT("NotTracingText","Not currently tracing."));
+	}
+		
+	return DescBuilder.ToText();
+}
+
+FSlateColor SInsightsStatusBarWidget::GetRecordingButtonColor() const
+{
+	if (!FTraceAuxiliary::IsConnected())
+	{
+		return FStyleColors::White;
+	}
+
+	return FStyleColors::Error;
+}
+
+FSlateColor SInsightsStatusBarWidget::GetRecordingButtonOutlineColor() const
+{
+	if (!FTraceAuxiliary::IsConnected())
+	{
+		ConnectionStartTime = FSlateApplication::Get().GetCurrentTime();
+		return FLinearColor::White.CopyWithNewOpacity(0.5f);
+	}
+
+	double ElapsedTime = FSlateApplication::Get().GetCurrentTime() - ConnectionStartTime;
+	return FStyleColors::Error.GetColor(FWidgetStyle()).CopyWithNewOpacity(0.5f + 0.5f * FMath::MakePulsatingValue(ElapsedTime, 0.5f));
+}
+
+FText SInsightsStatusBarWidget::GetRecordingButtonTooltipText() const
+{
+	if (!FTraceAuxiliary::IsConnected())
+	{
+		return LOCTEXT("StartTracing", "Start tracing. The trace destination is set from the menu.");
+	}
+
+	return LOCTEXT("StopTracing", "Stop Tracing.");
+}
+
+void SInsightsStatusBarWidget::SendSnapshotNotification()
+{
+	FNotificationInfo Info(LOCTEXT("SnapshotSaved", "Insights Snapshot saved."));
+	Info.bFireAndForget = true;
+	Info.FadeOutDuration = 1.0f;
+	Info.ExpireDuration = 4.0f;
+
+	Info.Text = LOCTEXT("SnapshotSavedHeading", "Insights Snapshot saved.");
+	Info.SubText = LOCTEXT("SnapshotSavedText", "A snapshot .utrace with the most recent events has been saved to your Saved/Profiling/ directory.");
+
+	FSlateNotificationManager::Get().AddNotification(Info);
+}
+
+void SInsightsStatusBarWidget::SendTraceStartedNotification()
+{
+	FNotificationInfo Info(LOCTEXT("TracingStarted", "Tracing started."));
+	Info.bFireAndForget = true;
+	Info.FadeOutDuration = 1.0f;
+	Info.ExpireDuration = 4.0f;
+
+	Info.Text = LOCTEXT("TracingStartedHeading", "Tracing started.");
+	Info.SubText = FText::Format(
+		LOCTEXT("TracingStartedText", "Trace is now active and saving to the following location (file or tracestore):\n{0}")
+		, FText::FromString(FTraceAuxiliary::GetTraceDestination())
+	);
+	FSlateNotificationManager::Get().AddNotification(Info);
+}
+
+void SInsightsStatusBarWidget::LaunchUnrealInsights_OnClicked()
+{
+	FUnrealInsightsLauncher::Get()->StartUnrealInsights(FUnrealInsightsLauncher::Get()->GetInsightsApplicationPath());
+}
+
+void SInsightsStatusBarWidget::OpenLiveSession_OnClicked()
+{
+	FUnrealInsightsLauncher::Get()->TryOpenTraceFromDestination(FTraceAuxiliary::GetTraceDestination());
+}
+
+
+void SInsightsStatusBarWidget::OpenProfilingDirectory_OnClicked()
+{
+	FString FullPath(FPaths::ConvertRelativePathToFull(FPaths::ProfilingDir()));
+	FPlatformProcess::ExploreFolder(*FullPath);
+}
+
+void SInsightsStatusBarWidget::OpenTraceStoreDirectory_OnClicked()
+{
+	UE::Trace::FStoreClient* StoreClient = UE::Trace::FStoreClient::Connect(TEXT("localhost"));
+
+	if (!StoreClient)
+	{
+		// TODO: Add Error Message
+		return;
+	}
+
+	const UE::Trace::FStoreClient::FStatus* Status = StoreClient->GetStatus();
+	if (!Status)
+	{
+		// TODO: Add Error Message
+		return;
+	}
+
+	FString Path(Status->GetStoreDir());
+
+	FString FullPath(FPaths::ConvertRelativePathToFull(Path));
+	FPlatformProcess::ExploreFolder(*FullPath);
+}
+
+void SInsightsStatusBarWidget::SetTraceDestination_Execute(ETraceDestination InDestination)
+{
+	TraceDestination = InDestination;
+}
+
+bool SInsightsStatusBarWidget::SetTraceDestination_IsChecked(ETraceDestination InDestination)
+{
+	return InDestination == TraceDestination;
+}
+
+bool SInsightsStatusBarWidget::SetTraceDestination_CanExecute()
+{
+	if (!FTraceAuxiliary::IsConnected())
+	{
+		return true;
+	}
+
+	return false;
+}
+
+void SInsightsStatusBarWidget::SaveSnapshot()
+{
+	FTraceAuxiliary::WriteSnapshot(nullptr);
+	SendSnapshotNotification();
+}
+
+bool SInsightsStatusBarWidget::SaveSnapshot_CanExecute()
+{
+	return true;
+}
+
+FText SInsightsStatusBarWidget::GetTraceMenuItemText() const
+{
+	if (FTraceAuxiliary::IsConnected())
+	{
+		return LOCTEXT("StopTraceButtonText", "Stop Trace");
+	}
+
+	return LOCTEXT("StartTraceButtonText", "Start Trace");
+}
+
+FText SInsightsStatusBarWidget::GetTraceMenuItemTooltipText() const
+{
+	if (FTraceAuxiliary::IsConnected())
+	{
+		return LOCTEXT("StopTraceButtonTooltip", "Stop tracing");
+	}
+
+	return LOCTEXT("StartTraceButtonTooltip", "Start tracing to the selected trace destination.");
+}
+
+void SInsightsStatusBarWidget::ToggleTracing_OnClicked()
+{
+	if (FTraceAuxiliary::IsConnected())
+	{
+		FTraceAuxiliary::Stop();
+	}
+	else
+	{
+		StartTracing();
+		SendTraceStartedNotification();
+	}
+}
+
+void SInsightsStatusBarWidget::StartTracing()
+{
+	if (TraceDestination == ETraceDestination::TraceStore)
+	{
+		FTraceAuxiliary::Start(FTraceAuxiliary::EConnectionType::Network, TEXT("localhost"), Channels);
+	}
+	else if (TraceDestination == ETraceDestination::File)
+	{
+		FTraceAuxiliary::Start(FTraceAuxiliary::EConnectionType::File, nullptr, Channels);
+	}
+}
+
+EVisibility SInsightsStatusBarWidget::GetStartTraceIconVisibility() const
+{
+	if (GetStopTraceIconVisibility() == EVisibility::Hidden)
+	{
+		return EVisibility::Visible;
+	}
+
+	return EVisibility::Hidden;
+}
+
+EVisibility SInsightsStatusBarWidget::GetStopTraceIconVisibility() const
+{
+	if (bIsTraceRecordButtonHovered && FTraceAuxiliary::IsConnected())
+	{
+		return EVisibility::Visible;
+	}
+
+	return EVisibility::Hidden;
+}
+
+void SInsightsStatusBarWidget::SetTraceChannels(const TCHAR* InChannels)
+{
+	Channels = InChannels;
+}
+
+bool SInsightsStatusBarWidget::IsPresetSet(const TCHAR* InChannels) const
+{
+	return Channels == InChannels;
+}
+
+#undef LOCTEXT_NAMESPACE
diff --git a/Engine/Plugins/TraceUtilities/Source/EditorTraceUtilities/Private/InsightsStatusBar/SInsightsStatusBar.h b/Engine/Plugins/TraceUtilities/Source/EditorTraceUtilities/Private/InsightsStatusBar/SInsightsStatusBar.h
new file mode 100644
--- /dev/null
+++ b/Engine/Plugins/TraceUtilities/Source/EditorTraceUtilities/Private/InsightsStatusBar/SInsightsStatusBar.h
@@ -0,0 +1,82 @@
+// Copyright Epic Games, Inc. All Rights Reserved.
+
+#pragma once
+
+#include "CoreMinimal.h"
+#include "Widgets/DeclarativeSyntaxSupport.h"
+#include "Widgets/SWidget.h"
+#include "Widgets/SCompoundWidget.h"
+
+class FMenuBuilder;
+
+TSharedRef<SWidget> CreateInsightsStatusBarWidget();
+
+enum class ETraceDestination : uint32
+{
+	TraceStore = 0,
+	File = 1
+};
+
+/**
+ *  Status bar widget for Unreal Insights.
+ *  Shows buttons to start tracing either to a file or to the trace store and allows saving a snapshot to file.
+ */
+class SInsightsStatusBarWidget : public SCompoundWidget
+{
+	SLATE_BEGIN_ARGS(SInsightsStatusBarWidget) {}
+	SLATE_END_ARGS()
+
+	void Construct(const FArguments& InArgs);
+
+private:
+
+	FText	GetTitleToolTipText() const;
+	
+	FSlateColor GetRecordingButtonColor() const;
+	FSlateColor GetRecordingButtonOutlineColor() const;
+	FText GetRecordingButtonTooltipText() const;
+
+	void LaunchUnrealInsights_OnClicked();
+	
+	void OpenLiveSession_OnClicked();
+
+	void OpenProfilingDirectory_OnClicked();
+	void OpenTraceStoreDirectory_OnClicked();
+
+	void SetTraceDestination_Execute(ETraceDestination InDestination);
+	bool SetTraceDestination_CanExecute();
+	bool SetTraceDestination_IsChecked(ETraceDestination InDestination);
+
+	void SaveSnapshot();
+	bool SaveSnapshot_CanExecute();
+
+	FText GetTraceMenuItemText() const;
+	FText GetTraceMenuItemTooltipText() const;
+
+	void ToggleTracing_OnClicked();
+
+	EVisibility GetStartTraceIconVisibility() const;
+	EVisibility GetStopTraceIconVisibility() const;
+
+	void StartTracing();
+
+	TSharedRef<SWidget> MakeTraceMenu();
+	void Channels_BuildMenu(FMenuBuilder& MenuBuilder);
+
+	static void SendSnapshotNotification();
+	static void SendTraceStartedNotification();
+
+	void SetTraceChannels(const TCHAR* InChannels);
+	bool IsPresetSet(const TCHAR* InChannels) const;
+
+private:
+	static const TCHAR* DefaultPreset;
+	static const TCHAR* MemoryPreset;
+	static const TCHAR* TaskGraphPreset;
+	static const TCHAR* ContextSwitchesPreset;
+
+	ETraceDestination TraceDestination = ETraceDestination::TraceStore;
+	bool bIsTraceRecordButtonHovered = false;
+	const TCHAR* Channels;
+	mutable double ConnectionStartTime = 0.0f;
+};
diff --git a/Engine/Plugins/TraceUtilities/TraceUtilities.uplugin b/Engine/Plugins/TraceUtilities/TraceUtilities.uplugin
--- a/Engine/Plugins/TraceUtilities/TraceUtilities.uplugin
+++ b/Engine/Plugins/TraceUtilities/TraceUtilities.uplugin
@@ -9,7 +9,7 @@
   "CreatedByURL": "https://epicgames.com",
   "DocsURL": "",
   "MarketplaceURL": "",
-  "CanContainContent": false,
+  "CanContainContent": true,
   "IsBetaVersion": false,
   "IsExperimentalVersion": false,
   "Installed": false,
diff --git a/Engine/Source/Runtime/Core/Private/ProfilingDebugging/TraceAuxiliary.cpp b/Engine/Source/Runtime/Core/Private/ProfilingDebugging/TraceAuxiliary.cpp
--- a/Engine/Source/Runtime/Core/Private/ProfilingDebugging/TraceAuxiliary.cpp
+++ b/Engine/Source/Runtime/Core/Private/ProfilingDebugging/TraceAuxiliary.cpp
@@ -130,7 +130,7 @@ private:
 	ChannelSet				CommandlineChannels;
 	FString					TraceDest;
 	FTraceAuxiliary::EConnectionType TraceType = FTraceAuxiliary::EConnectionType::None;
-	EState					State = EState::Stopped;
+	std::atomic<EState>		State = EState::Stopped;
 	bool					bTruncateFile = false;
 	bool					bWorkerThreadStarted = false;
 	FString					PausedPreset;
@@ -248,7 +248,7 @@ void FTraceAuxiliaryImpl::AddChannel(const TCHAR* Name)
 	FChannelEntry& Value = CommandlineChannels.Add(Hash, {});
 	Value.Name = Name;
 
-	if (State >= EState::Tracing && !Value.bActive)
+	if (State.load() >= EState::Tracing && !Value.bActive)
 	{
 		Value.bActive = EnableChannel(*Value.Name);
 	}
@@ -265,7 +265,7 @@ void FTraceAuxiliaryImpl::RemoveChannel(const TCHAR* Name)
 		return;
 	}
 
-	if (State >= EState::Tracing && Channel.bActive)
+	if (State.load() >= EState::Tracing && Channel.bActive)
 	{
 		DisableChannel(*Channel.Name);
 		Channel.bActive = false;
@@ -320,7 +320,7 @@ bool FTraceAuxiliaryImpl::Connect(ETraceConnectType Type, const TCHAR* Parameter
 		return false;
 	}
 
-	State = EState::Tracing;
+	State.store(EState::Tracing);
 	return true;
 }
 
@@ -339,7 +339,7 @@ bool FTraceAuxiliaryImpl::Stop()
 
 	FTraceAuxiliary::OnTraceStopped.Broadcast(TraceType, TraceDest);
 
-	State = EState::Stopped;
+	State.store(EState::Stopped);
 	TraceType = FTraceAuxiliary::EConnectionType::None;
 	TraceDest.Reset();
 	return true;
@@ -583,7 +583,7 @@ const TCHAR* FTraceAuxiliaryImpl::GetDest() const
 ////////////////////////////////////////////////////////////////////////////////
 bool FTraceAuxiliaryImpl::IsConnected() const
 {
-	return State == EState::Tracing;
+	return State.load() == EState::Tracing;
 }
 
 ////////////////////////////////////////////////////////////////////////////////
