Fix LOCTEXT conflicts in SInsightsStatusBar.cpp

diff --git a/Engine/Plugins/TraceUtilities/Source/EditorTraceUtilities/Private/InsightsStatusBar/SInsightsStatusBar.cpp b/Engine/Plugins/TraceUtilities/Source/EditorTraceUtilities/Private/InsightsStatusBar/SInsightsStatusBar.cpp
--- a/Engine/Plugins/TraceUtilities/Source/EditorTraceUtilities/Private/InsightsStatusBar/SInsightsStatusBar.cpp
+++ b/Engine/Plugins/TraceUtilities/Source/EditorTraceUtilities/Private/InsightsStatusBar/SInsightsStatusBar.cpp
@@ -167,7 +167,7 @@ TSharedRef<SWidget> SInsightsStatusBarWidget::MakeTraceMenu()
 	{
 		MenuBuilder.AddMenuEntry(
 			LOCTEXT("ServerLabel", "Trace Store"),
-			LOCTEXT("FileLabelDesc", "Set the trace store as the trace destination."),
+			LOCTEXT("ServerLabelDesc", "Set the trace store as the trace destination."),
 			FSlateIcon(),
 			FUIAction(FExecuteAction::CreateSP(this, &SInsightsStatusBarWidget::SetTraceDestination_Execute, ETraceDestination::TraceStore),
 					  FCanExecuteAction::CreateSP(this, &SInsightsStatusBarWidget::SetTraceDestination_CanExecute),
@@ -177,7 +177,7 @@ TSharedRef<SWidget> SInsightsStatusBarWidget::MakeTraceMenu()
 		);
 
 		MenuBuilder.AddMenuEntry(
-			LOCTEXT("ServerLabel", "File"),
+			LOCTEXT("FileLabel", "File"),
 			LOCTEXT("FileLabelDesc", "Set file as the trace destination."),
 			FSlateIcon(),
 			FUIAction(FExecuteAction::CreateSP(this, &SInsightsStatusBarWidget::SetTraceDestination_Execute, ETraceDestination::File),
@@ -193,7 +193,7 @@ TSharedRef<SWidget> SInsightsStatusBarWidget::MakeTraceMenu()
 	{
 		MenuBuilder.AddMenuEntry(
 			LOCTEXT("SaveSnapshotLabel", "Save Trace Snapshot"),
-			LOCTEXT("OpenInsightsTooltip", "Save the current trace buffer to file."),
+			LOCTEXT("SaveSnapshotTooltip", "Save the current trace buffer to file."),
 			FSlateIcon(FEditorTraceUtilitiesStyle::Get().GetStyleSetName(), "Icons.TraceSnapshot"),
 			FUIAction(FExecuteAction::CreateSP(this, &SInsightsStatusBarWidget::SaveSnapshot),
 					  FCanExecuteAction::CreateSP(this, &SInsightsStatusBarWidget::SaveSnapshot_CanExecute)),
