Add options to the Insights Status Bar Widget

diff --git a/Engine/Plugins/TraceUtilities/Config/BaseTraceUtilities.ini b/Engine/Plugins/TraceUtilities/Config/BaseTraceUtilities.ini
new file mode 100644
--- /dev/null
+++ b/Engine/Plugins/TraceUtilities/Config/BaseTraceUtilities.ini
@@ -0,0 +1,4 @@
+[EditorTraceUtilities]
+OpenLiveSesssionOnTraceStart=false
+OpenInsightsAfterTrace=false
+ShowInExplorerAfterTrace=false
\ No newline at end of file
diff --git a/Engine/Plugins/TraceUtilities/Source/EditorTraceUtilities/Private/EditorTraceUtilities.cpp b/Engine/Plugins/TraceUtilities/Source/EditorTraceUtilities/Private/EditorTraceUtilities.cpp
--- a/Engine/Plugins/TraceUtilities/Source/EditorTraceUtilities/Private/EditorTraceUtilities.cpp
+++ b/Engine/Plugins/TraceUtilities/Source/EditorTraceUtilities/Private/EditorTraceUtilities.cpp
@@ -2,6 +2,7 @@
 
 #include "EditorTraceUtilities.h"
 
+#include "Misc/ConfigContext.h"
 #include "ToolMenus.h"
 #include "UnrealInsightsLauncher.h"
 
@@ -14,6 +15,8 @@ DEFINE_LOG_CATEGORY(LogTraceUtilities)
   */
 void RegisterInsightsStatusWidgetWithToolMenu();
 
+FString FEditorTraceUtilitiesModule::EditorTraceUtilitiesIni;
+
 void FEditorTraceUtilitiesModule::StartupModule()
 {
 	LLM_SCOPE_BYNAME(TEXT("Insights"));
@@ -23,6 +26,8 @@ void FEditorTraceUtilitiesModule::StartupModule()
 			&FUnrealInsightsLauncher::RegisterMenus));
 
 	RegisterInsightsStatusWidgetWithToolMenu();
+
+	FConfigContext::ReadIntoGConfig().Load(TEXT("TraceUtilities"), EditorTraceUtilitiesIni);
 }
 
 void FEditorTraceUtilitiesModule::ShutdownModule()
diff --git a/Engine/Plugins/TraceUtilities/Source/EditorTraceUtilities/Private/InsightsStatusBar/SInsightsStatusBar.cpp b/Engine/Plugins/TraceUtilities/Source/EditorTraceUtilities/Private/InsightsStatusBar/SInsightsStatusBar.cpp
--- a/Engine/Plugins/TraceUtilities/Source/EditorTraceUtilities/Private/InsightsStatusBar/SInsightsStatusBar.cpp
+++ b/Engine/Plugins/TraceUtilities/Source/EditorTraceUtilities/Private/InsightsStatusBar/SInsightsStatusBar.cpp
@@ -2,13 +2,17 @@
 
 #include "SInsightsStatusBar.h"
 
+#include "CoreGlobals.h"
 #include "EditorTraceUtilitiesStyle.h"
+#include "EditorTraceUtilities.h"
 #include "Framework/Application/SlateApplication.h"
 #include "Framework/MultiBox/MultiBoxBuilder.h"
 #include "Framework/Notifications/NotificationManager.h"
+#include "GenericPlatform/GenericPlatformFile.h"
+#include "HAL/FileManager.h"
 #include "Math/Color.h"
+#include "Misc/ConfigCacheIni.h"
 #include "Modules/ModuleManager.h"
-#include "ProfilingDebugging/TraceAuxiliary.h"
 #include "Styling/StyleColors.h"
 #include "ToolMenus.h"
 #include "Trace/StoreClient.h"
@@ -27,6 +31,11 @@ const TCHAR* SInsightsStatusBarWidget::MemoryPreset= TEXT("default,memory");
 const TCHAR* SInsightsStatusBarWidget::TaskGraphPreset = TEXT("default,task");
 const TCHAR* SInsightsStatusBarWidget::ContextSwitchesPreset = TEXT("default,contextswitches");
 
+const TCHAR* SInsightsStatusBarWidget::SettingsCategory = TEXT("EditorTraceUtilities");
+const TCHAR* SInsightsStatusBarWidget::OpenLiveSessionOnTraceStartSettingName = TEXT("OpenLiveSessionOnTraceStart");
+const TCHAR* SInsightsStatusBarWidget::OpenInsightsAfterTraceSettingName = TEXT("OpenInsightsAfterTrace");
+const TCHAR* SInsightsStatusBarWidget::ShowInExplorerAfterTraceSettingName = TEXT("ShowInExplorerAfterTrace");
+
 TSharedRef<SWidget> CreateInsightsStatusBarWidget()
 {
 	return SNew(SInsightsStatusBarWidget);
@@ -36,7 +45,7 @@ void RegisterInsightsStatusWidgetWithToolMenu()
 {
 	UToolMenu* Menu = UToolMenus::Get()->ExtendMenu(TEXT("LevelEditor.StatusBar.ToolBar"));
 
-	FToolMenuSection& InsightsSection = Menu->AddSection(TEXT("Insights"), FText::GetEmpty(), FToolMenuInsert(TEXT("Compile"), EToolMenuInsertType::Before));
+	FToolMenuSection& InsightsSection = Menu->AddSection(TEXT("Insights"), FText::GetEmpty(), FToolMenuInsert(TEXT("DDC"), EToolMenuInsertType::Before));
 
 	InsightsSection.AddEntry(
 		FToolMenuEntry::InitWidget(TEXT("InsightsStatusBar"), CreateInsightsStatusBarWidget(), FText::GetEmpty(), true, false)
@@ -142,26 +151,49 @@ void SInsightsStatusBarWidget::Construct(const FArguments& InArgs)
 		]
 	];
 
+	if (FTraceAuxiliary::GetConnectionType() == FTraceAuxiliary::EConnectionType::Network)
+	{
+		TraceDestination = ETraceDestination::TraceStore;
+	}
+	if (FTraceAuxiliary::GetConnectionType() == FTraceAuxiliary::EConnectionType::File)
+	{
+		TraceDestination = ETraceDestination::File;
+	}
+
 	Channels = DefaultPreset;
+	FTraceAuxiliary::OnTraceStarted.AddSP(this, &SInsightsStatusBarWidget::OnTraceStarted);
+	FTraceAuxiliary::OnTraceStopped.AddSP(this, &SInsightsStatusBarWidget::OnTraceStopped);
+	FTraceAuxiliary::OnSnapshotSaved.AddSP(this, &SInsightsStatusBarWidget::OnSnapshotSaved);
 }
 
 TSharedRef<SWidget> SInsightsStatusBarWidget::MakeTraceMenu()
 {
 	FMenuBuilder MenuBuilder(true, nullptr);
 
-	// To be activated
-	//MenuBuilder.BeginSection("TraceDataFiltering", LOCTEXT("TraceMenu_Section_DataFiltering", "Trace Data Filtering"));
-	//{
-	//	MenuBuilder.AddSubMenu
-	//	(
-	//		LOCTEXT("Channels", "Channels"),
-	//		LOCTEXT("Channels_Desc", "Select what trace channels to enable when tracing."),
-	//		FNewMenuDelegate::CreateSP(this, &SInsightsStatusBarWidget::Channels_BuildMenu),
-	//		false,
-	//		FSlateIcon(FEditorTraceUtilitiesStyle::Get().GetStyleSetName(), ("Icons.Trace"))
-	//	);
-	//}
-	//MenuBuilder.EndSection();
+	MenuBuilder.BeginSection("TraceDataFiltering", LOCTEXT("TraceMenu_Section_DataFiltering", "Trace Data Filtering"));
+	{
+		MenuBuilder.AddMenuEntry(
+			LOCTEXT("StatNamedEventsLabel", "Stat Named Events"),
+			LOCTEXT("StatNamedEventsDesc", "Enable or disable named events in the stats system."),
+			FSlateIcon(),
+			FUIAction(FExecuteAction::CreateLambda([]() { GCycleStatsShouldEmitNamedEvents = GCycleStatsShouldEmitNamedEvents == 0 ? 1 : 0; }),
+					  FCanExecuteAction::CreateLambda([]() { return true; }),
+					  FIsActionChecked::CreateLambda([]() { return GCycleStatsShouldEmitNamedEvents > 0; })),
+			NAME_None,
+			EUserInterfaceActionType::ToggleButton
+		);
+
+		// To be activated 
+		//	MenuBuilder.AddSubMenu
+		//	(
+		//		LOCTEXT("Channels", "Channels"),
+		//		LOCTEXT("Channels_Desc", "Select what trace channels to enable when tracing."),
+		//		FNewMenuDelegate::CreateSP(this, &SInsightsStatusBarWidget::Channels_BuildMenu),
+		//		false,
+		//		FSlateIcon(FEditorTraceUtilitiesStyle::Get().GetStyleSetName(), ("Icons.Trace"))
+		//	);
+	}
+	MenuBuilder.EndSection();
 
 	MenuBuilder.BeginSection("TraceDestination", LOCTEXT("TraceMenu_Section_Destination", "Trace Destination"));
 	{
@@ -212,32 +244,54 @@ TSharedRef<SWidget> SInsightsStatusBarWidget::MakeTraceMenu()
 	}
 	MenuBuilder.EndSection();
 
-	MenuBuilder.BeginSection("Insights", LOCTEXT("TraceMenu_Section_Insights", "Insights"));
+	MenuBuilder.BeginSection("Options", LOCTEXT("TraceMenu_Section_Options", "Options"));
 	{
 		MenuBuilder.AddMenuEntry(
-			LOCTEXT("OpenInsightsLabel", "Unreal Insights (Session Browser)"),
-			LOCTEXT("OpenInsightsTooltip", "Launch the Unreal Insights Session Browser."),
-			FSlateIcon(FAppStyle::Get().GetStyleSetName(), "UnrealInsights.MenuIcon"),
-			FUIAction(FExecuteAction::CreateSP(this, &SInsightsStatusBarWidget::LaunchUnrealInsights_OnClicked)),
+			LOCTEXT("OpenLiveSesssionOnTraceStart", "Open Live Session on Trace Start"),
+			LOCTEXT("OpenLiveSesssionOnTraceStartDesc", "When set, the live session will be automatically opened in Unreal Insights when tracing is started.\nThis option will only apply when tracing to the trace store."),
+			FSlateIcon(),
+			FUIAction(FExecuteAction::CreateSP(this, &SInsightsStatusBarWidget::ToggleBooleanSettingValue, OpenLiveSessionOnTraceStartSettingName),
+				FCanExecuteAction::CreateLambda([]() { return true; }),
+				FIsActionChecked::CreateSP(this, &SInsightsStatusBarWidget::GetBooleanSettingValue, OpenLiveSessionOnTraceStartSettingName)),
 			NAME_None,
-			EUserInterfaceActionType::Button
+			EUserInterfaceActionType::ToggleButton
 		);
 
 		MenuBuilder.AddMenuEntry(
-			LOCTEXT("OpenLiveSessionLabel", "Open Live Session"),
-			LOCTEXT("OpenLiveSessionTooltip", "Opening the live session is possible only while tracing to the trace store."),
-			FSlateIcon(FEditorTraceUtilitiesStyle::Get().GetStyleSetName(), "Icons.OpenLiveSession"),
-			FUIAction(FExecuteAction::CreateSP(this, &SInsightsStatusBarWidget::OpenLiveSession_OnClicked),
-					  FCanExecuteAction::CreateLambda([this]() { return FTraceAuxiliary::IsConnected() && this->TraceDestination == ETraceDestination::TraceStore; })),
+			LOCTEXT("OpenInsightsAfterTrace", "Open Insights after Trace"),
+			LOCTEXT("OpenInsightsAfterTraceDesc", "When set, the session will be automatically opened in Unreal Insights when tracing is stopped or when a snapshot is saved."),
+			FSlateIcon(),
+			FUIAction(FExecuteAction::CreateSP(this, &SInsightsStatusBarWidget::ToggleBooleanSettingValue, OpenInsightsAfterTraceSettingName),
+				FCanExecuteAction::CreateLambda([]() { return true; }),
+				FIsActionChecked::CreateSP(this, &SInsightsStatusBarWidget::GetBooleanSettingValue, OpenInsightsAfterTraceSettingName)),
 			NAME_None,
-			EUserInterfaceActionType::Button
+			EUserInterfaceActionType::ToggleButton
+		);
+
+		MenuBuilder.AddMenuEntry(
+			LOCTEXT("ShowInExplorerAfterTrace", "Show in Explorer after Trace"),
+			LOCTEXT("ShowInExplorerAfterTraceDesc", "When set, folder containing the recorded session will be opened automatically when trace is stopped or when a snapshot is saved."),
+			FSlateIcon(),
+			FUIAction(FExecuteAction::CreateSP(this, &SInsightsStatusBarWidget::ToggleBooleanSettingValue, ShowInExplorerAfterTraceSettingName),
+				FCanExecuteAction::CreateLambda([]() { return true; }),
+				FIsActionChecked::CreateSP(this, &SInsightsStatusBarWidget::GetBooleanSettingValue, ShowInExplorerAfterTraceSettingName)),
+			NAME_None,
+			EUserInterfaceActionType::ToggleButton
 		);
 	}
 	MenuBuilder.EndSection();
 
-
 	MenuBuilder.BeginSection("Locations", LOCTEXT("TraceMenu_Section_Locations", "Locations"));
 	{
+		MenuBuilder.AddMenuEntry(
+			LOCTEXT("OpenTraceStoreLabel", "Open Trace Store Directory"),
+			LOCTEXT("OpenTraceStoreTooltip", "Open Trace Store Directory. This is the location where traces saved to the trace server are stored."),
+			FSlateIcon(FAppStyle::Get().GetStyleSetName(), "Icons.FolderOpen"),
+			FUIAction(FExecuteAction::CreateSP(this, &SInsightsStatusBarWidget::OpenTraceStoreDirectory_OnClicked)),
+			NAME_None,
+			EUserInterfaceActionType::Button
+		);
+
 		MenuBuilder.AddMenuEntry(
 			LOCTEXT("OpenProfilingDirectoryLabel", "Open Profiling Directory"),
 			LOCTEXT("OpenProfilingDirectoryTooltip", "Opens the profiling directory of the current project. This is the location where traces to file are stored."),
@@ -246,12 +300,26 @@ TSharedRef<SWidget> SInsightsStatusBarWidget::MakeTraceMenu()
 			NAME_None,
 			EUserInterfaceActionType::Button
 		);
+	}
+	MenuBuilder.EndSection();
 
+	MenuBuilder.BeginSection("Insights", LOCTEXT("TraceMenu_Section_Insights", "Insights"));
+	{
 		MenuBuilder.AddMenuEntry(
-			LOCTEXT("OpenTraceStoreLabel", "Open Trace Store Directory"),
-			LOCTEXT("OpenTraceStoreTooltip", "Open Trace Store Directory. This is the location where traces saved to the trace server are stored."),
-			FSlateIcon(FAppStyle::Get().GetStyleSetName(), "Icons.FolderOpen"),
-			FUIAction(FExecuteAction::CreateSP(this, &SInsightsStatusBarWidget::OpenTraceStoreDirectory_OnClicked)),
+			LOCTEXT("OpenInsightsLabel", "Unreal Insights (Session Browser)"),
+			LOCTEXT("OpenInsightsTooltip", "Launch the Unreal Insights Session Browser."),
+			FSlateIcon(FAppStyle::Get().GetStyleSetName(), "UnrealInsights.MenuIcon"),
+			FUIAction(FExecuteAction::CreateSP(this, &SInsightsStatusBarWidget::LaunchUnrealInsights_OnClicked)),
+			NAME_None,
+			EUserInterfaceActionType::Button
+		);
+
+		MenuBuilder.AddMenuEntry(
+			LOCTEXT("OpenLiveSessionLabel", "Open Live Session"),
+			LOCTEXT("OpenLiveSessionTooltip", "Opening the live session is possible only while tracing to the trace store."),
+			FSlateIcon(FEditorTraceUtilitiesStyle::Get().GetStyleSetName(), "Icons.OpenLiveSession"),
+			FUIAction(FExecuteAction::CreateSP(this, &SInsightsStatusBarWidget::OpenLiveSession_OnClicked),
+				FCanExecuteAction::CreateLambda([this]() { return FTraceAuxiliary::GetConnectionType() == FTraceAuxiliary::EConnectionType::Network; })),
 			NAME_None,
 			EUserInterfaceActionType::Button
 		);
@@ -397,17 +465,31 @@ void SInsightsStatusBarWidget::LaunchUnrealInsights_OnClicked()
 
 void SInsightsStatusBarWidget::OpenLiveSession_OnClicked()
 {
-	FUnrealInsightsLauncher::Get()->TryOpenTraceFromDestination(FTraceAuxiliary::GetTraceDestination());
+	OpenLiveSession();
 }
 
+void SInsightsStatusBarWidget::OpenLiveSession()
+{
+	FUnrealInsightsLauncher::Get()->TryOpenTraceFromDestination(FTraceAuxiliary::GetTraceDestination());
+}
 
 void SInsightsStatusBarWidget::OpenProfilingDirectory_OnClicked()
+{
+	OpenProfilingDirectory();
+}
+
+void SInsightsStatusBarWidget::OpenProfilingDirectory()
 {
 	FString FullPath(FPaths::ConvertRelativePathToFull(FPaths::ProfilingDir()));
 	FPlatformProcess::ExploreFolder(*FullPath);
 }
 
 void SInsightsStatusBarWidget::OpenTraceStoreDirectory_OnClicked()
+{
+	OpenTraceStoreDirectory(false);
+}
+
+void SInsightsStatusBarWidget::OpenTraceStoreDirectory(bool bSelectLastTrace)
 {
 	UE::Trace::FStoreClient* StoreClient = UE::Trace::FStoreClient::Connect(TEXT("localhost"));
 
@@ -426,6 +508,32 @@ void SInsightsStatusBarWidget::OpenTraceStoreDirectory_OnClicked()
 
 	FString Path(Status->GetStoreDir());
 
+	if (bSelectLastTrace)
+	{
+		FString MostRecentTraceName;
+		FDateTime LatestDateTime;
+
+		auto Visitor = [&MostRecentTraceName, &LatestDateTime](const TCHAR* Filename, const FFileStatData& StatData)
+		{
+			if (FPaths::GetExtension(Filename) == TEXT("utrace"))
+			{
+				if (LatestDateTime < StatData.ModificationTime)
+				{
+					LatestDateTime = StatData.ModificationTime;
+					MostRecentTraceName = Filename;
+				}
+			}
+			return true;
+		};
+
+		IFileManager::Get().IterateDirectoryStat(*Path, Visitor);
+
+		if (!MostRecentTraceName.IsEmpty())
+		{
+			Path = MostRecentTraceName;
+		}
+	}
+
 	FString FullPath(FPaths::ConvertRelativePathToFull(Path));
 	FPlatformProcess::ExploreFolder(*FullPath);
 }
@@ -536,4 +644,57 @@ bool SInsightsStatusBarWidget::IsPresetSet(const TCHAR* InChannels) const
 	return Channels == InChannels;
 }
 
+bool SInsightsStatusBarWidget::GetBooleanSettingValue(const TCHAR* InSettingName)
+{
+	bool Value = false;
+	GConfig->GetBool(SettingsCategory, InSettingName, Value, FEditorTraceUtilitiesModule::GetTraceUtilitiesIni());
+	return Value;
+}
+
+void SInsightsStatusBarWidget::ToggleBooleanSettingValue(const TCHAR* InSettingName)
+{
+	bool Value = false;
+	GConfig->GetBool(SettingsCategory, InSettingName, Value, FEditorTraceUtilitiesModule::GetTraceUtilitiesIni());
+	GConfig->SetBool(SettingsCategory, InSettingName, !Value, FEditorTraceUtilitiesModule::GetTraceUtilitiesIni());
+}
+
+void SInsightsStatusBarWidget::OnTraceStarted(FTraceAuxiliary::EConnectionType InTraceType, const FString& InTraceDestination)
+{
+	if (InTraceType == FTraceAuxiliary::EConnectionType::Network && GetBooleanSettingValue(OpenLiveSessionOnTraceStartSettingName))
+	{
+		OpenLiveSession();
+	}
+}
+
+void SInsightsStatusBarWidget::OnTraceStopped(FTraceAuxiliary::EConnectionType InTraceType, const FString& InTraceDestination)
+{
+	if (GetBooleanSettingValue(OpenInsightsAfterTraceSettingName))
+	{
+		OpenLiveSession();
+	}
+	if (GetBooleanSettingValue(ShowInExplorerAfterTraceSettingName))
+	{
+		if (InTraceType == FTraceAuxiliary::EConnectionType::Network)
+		{
+			OpenTraceStoreDirectory(true);
+		}
+		else if (InTraceType == FTraceAuxiliary::EConnectionType::File)
+		{
+			FPlatformProcess::ExploreFolder(*InTraceDestination);
+		}
+	}
+}
+
+void SInsightsStatusBarWidget::OnSnapshotSaved(const FString& InPath)
+{
+	if (GetBooleanSettingValue(OpenInsightsAfterTraceSettingName))
+	{
+		FUnrealInsightsLauncher::Get()->TryOpenTraceFromDestination(InPath);
+	}
+	if (GetBooleanSettingValue(ShowInExplorerAfterTraceSettingName))
+	{
+		FPlatformProcess::ExploreFolder(*InPath);
+	}
+}
+
 #undef LOCTEXT_NAMESPACE
diff --git a/Engine/Plugins/TraceUtilities/Source/EditorTraceUtilities/Private/InsightsStatusBar/SInsightsStatusBar.h b/Engine/Plugins/TraceUtilities/Source/EditorTraceUtilities/Private/InsightsStatusBar/SInsightsStatusBar.h
--- a/Engine/Plugins/TraceUtilities/Source/EditorTraceUtilities/Private/InsightsStatusBar/SInsightsStatusBar.h
+++ b/Engine/Plugins/TraceUtilities/Source/EditorTraceUtilities/Private/InsightsStatusBar/SInsightsStatusBar.h
@@ -3,6 +3,7 @@
 #pragma once
 
 #include "CoreMinimal.h"
+#include "ProfilingDebugging/TraceAuxiliary.h"
 #include "Widgets/DeclarativeSyntaxSupport.h"
 #include "Widgets/SWidget.h"
 #include "Widgets/SCompoundWidget.h"
@@ -29,7 +30,6 @@ class SInsightsStatusBarWidget : public SCompoundWidget
 	void Construct(const FArguments& InArgs);
 
 private:
-
 	FText	GetTitleToolTipText() const;
 	
 	FSlateColor GetRecordingButtonColor() const;
@@ -39,9 +39,13 @@ private:
 	void LaunchUnrealInsights_OnClicked();
 	
 	void OpenLiveSession_OnClicked();
+	void OpenLiveSession();
 
 	void OpenProfilingDirectory_OnClicked();
+	void OpenProfilingDirectory();
+
 	void OpenTraceStoreDirectory_OnClicked();
+	void OpenTraceStoreDirectory(bool bSelectLastTrace);
 
 	void SetTraceDestination_Execute(ETraceDestination InDestination);
 	bool SetTraceDestination_CanExecute();
@@ -69,12 +73,24 @@ private:
 	void SetTraceChannels(const TCHAR* InChannels);
 	bool IsPresetSet(const TCHAR* InChannels) const;
 
+	bool GetBooleanSettingValue(const TCHAR* InSettingName);
+	void ToggleBooleanSettingValue(const TCHAR* InSettingName);
+
+	void OnTraceStarted(FTraceAuxiliary::EConnectionType TraceType, const FString& TraceDestination);
+	void OnTraceStopped(FTraceAuxiliary::EConnectionType TraceType, const FString& TraceDestination);
+	void OnSnapshotSaved(const FString& InPath);
+
 private:
 	static const TCHAR* DefaultPreset;
 	static const TCHAR* MemoryPreset;
 	static const TCHAR* TaskGraphPreset;
 	static const TCHAR* ContextSwitchesPreset;
 
+	static const TCHAR* SettingsCategory;
+	static const TCHAR* OpenLiveSessionOnTraceStartSettingName;
+	static const TCHAR* OpenInsightsAfterTraceSettingName;
+	static const TCHAR* ShowInExplorerAfterTraceSettingName;
+
 	ETraceDestination TraceDestination = ETraceDestination::TraceStore;
 	bool bIsTraceRecordButtonHovered = false;
 	const TCHAR* Channels;
diff --git a/Engine/Plugins/TraceUtilities/Source/EditorTraceUtilities/Public/EditorTraceUtilities.h b/Engine/Plugins/TraceUtilities/Source/EditorTraceUtilities/Public/EditorTraceUtilities.h
--- a/Engine/Plugins/TraceUtilities/Source/EditorTraceUtilities/Public/EditorTraceUtilities.h
+++ b/Engine/Plugins/TraceUtilities/Source/EditorTraceUtilities/Public/EditorTraceUtilities.h
@@ -13,6 +13,9 @@ public:
     virtual void StartupModule() override;
     virtual void ShutdownModule() override;
     
+	static const FString& GetTraceUtilitiesIni() { return EditorTraceUtilitiesIni; }
+
 private:
     FDelegateHandle RegisterStartupCallbackHandle;
+	static FString EditorTraceUtilitiesIni;
 };
diff --git a/Engine/Source/Runtime/Core/Private/ProfilingDebugging/TraceAuxiliary.cpp b/Engine/Source/Runtime/Core/Private/ProfilingDebugging/TraceAuxiliary.cpp
--- a/Engine/Source/Runtime/Core/Private/ProfilingDebugging/TraceAuxiliary.cpp
+++ b/Engine/Source/Runtime/Core/Private/ProfilingDebugging/TraceAuxiliary.cpp
@@ -80,25 +80,26 @@ enum class ETraceConnectType
 class FTraceAuxiliaryImpl
 {
 public:
-	const TCHAR*			GetDest() const;
-	bool					IsConnected() const;
-	void					GetActiveChannelsString(FStringBuilderBase& String) const;
-	void					AddCommandlineChannels(const TCHAR* ChannelList);
-	void					ResetCommandlineChannels();
-	bool					HasCommandlineChannels() const { return !CommandlineChannels.IsEmpty(); }
-	void					EnableChannels(const TCHAR* ChannelList);
-	void					DisableChannels(const TCHAR* ChannelList);
-	bool					Connect(ETraceConnectType Type, const TCHAR* Parameter, const FTraceAuxiliary::FLogCategoryAlias& LogCategory);
-	bool					Stop();
-	void					ResumeChannels();
-	void					PauseChannels();
-	void					EnableCommandlineChannels();
-	void					EnableCommandlineChannelsPostInitialize();
-	void					SetTruncateFile(bool bTruncateFile);
-	void					UpdateCsvStats() const;
-	void					StartWorkerThread();
-	void					StartEndFramePump();
-	bool					WriteSnapshot(const TCHAR* InFilePath, const FTraceAuxiliary::FLogCategoryAlias& LogCategory);
+	const TCHAR* GetDest() const;
+	bool IsConnected() const;
+	FTraceAuxiliary::EConnectionType GetConnectionType() const;
+	void GetActiveChannelsString(FStringBuilderBase& String) const;
+	void AddCommandlineChannels(const TCHAR* ChannelList);
+	void ResetCommandlineChannels();
+	bool HasCommandlineChannels() const { return !CommandlineChannels.IsEmpty(); }
+	void EnableChannels(const TCHAR* ChannelList);
+	void DisableChannels(const TCHAR* ChannelList);
+	bool Connect(ETraceConnectType Type, const TCHAR* Parameter, const FTraceAuxiliary::FLogCategoryAlias& LogCategory);
+	bool Stop();
+	void ResumeChannels();
+	void PauseChannels();
+	void EnableCommandlineChannels();
+	void EnableCommandlineChannelsPostInitialize();
+	void SetTruncateFile(bool bTruncateFile);
+	void UpdateCsvStats() const;
+	void StartWorkerThread();
+	void StartEndFramePump();
+	bool WriteSnapshot(const TCHAR* InFilePath, const FTraceAuxiliary::FLogCategoryAlias& LogCategory);
 
 	// True if this is parent process with forking requested before forking.
 	bool					IsParentProcessAndPreFork();
@@ -112,28 +113,28 @@ private:
 
 	struct FChannelEntry
 	{
-		FString				Name;
-		bool				bActive = false;
+		FString	Name;
+		bool bActive = false;
 	};
 
-	void					AddChannel(const TCHAR* Name);
-	void					RemoveChannel(const TCHAR* Name);
+	void AddChannel(const TCHAR* Name);
+	void RemoveChannel(const TCHAR* Name);
 	template <class T> void ForEachChannel(const TCHAR* ChannelList, bool bResolvePresets, T Callable);
-	static uint32			HashChannelName(const TCHAR* Name);
-	bool					EnableChannel(const TCHAR* Channel);
-	void					DisableChannel(const TCHAR* Channel);
-	bool					SendToHost(const TCHAR* Host, const FTraceAuxiliary::FLogCategoryAlias& LogCategory);
-	bool					WriteToFile(const TCHAR* Path, const FTraceAuxiliary::FLogCategoryAlias& LogCategory);
-	bool					FinalizeFilePath(const TCHAR* InPath, FString& OutPath, const FTraceAuxiliary::FLogCategoryAlias& LogCategory);
+	static uint32 HashChannelName(const TCHAR* Name);
+	bool EnableChannel(const TCHAR* Channel);
+	void DisableChannel(const TCHAR* Channel);
+	bool SendToHost(const TCHAR* Host, const FTraceAuxiliary::FLogCategoryAlias& LogCategory);
+	bool WriteToFile(const TCHAR* Path, const FTraceAuxiliary::FLogCategoryAlias& LogCategory);
+	bool FinalizeFilePath(const TCHAR* InPath, FString& OutPath, const FTraceAuxiliary::FLogCategoryAlias& LogCategory);
 
 	typedef TMap<uint32, FChannelEntry, TInlineSetAllocator<128>> ChannelSet;
-	ChannelSet				CommandlineChannels;
-	FString					TraceDest;
-	FTraceAuxiliary::EConnectionType TraceType = FTraceAuxiliary::EConnectionType::None;
-	std::atomic<EState>		State = EState::Stopped;
-	bool					bTruncateFile = false;
-	bool					bWorkerThreadStarted = false;
-	FString					PausedPreset;
+	ChannelSet CommandlineChannels;
+	FString TraceDest;
+	std::atomic<FTraceAuxiliary::EConnectionType> TraceType = FTraceAuxiliary::EConnectionType::None;
+	std::atomic<EState> State = EState::Stopped;
+	bool bTruncateFile = false;
+	bool bWorkerThreadStarted = false;
+	FString PausedPreset;
 };
 
 static FTraceAuxiliaryImpl GTraceAuxiliary;
@@ -291,7 +292,7 @@ bool FTraceAuxiliaryImpl::Connect(ETraceConnectType Type, const TCHAR* Parameter
 				UE_LOG_REF(LogCategory, Error, TEXT("Trace failed to connect (trace server: %s)!"), Parameter ? Parameter : TEXT(""));
 			}
 
-			TraceType = FTraceAuxiliary::EConnectionType::Network;
+			TraceType.store(FTraceAuxiliary::EConnectionType::Network);
 		}
 
 		else if (Type == ETraceConnectType::File)
@@ -306,12 +307,12 @@ bool FTraceAuxiliaryImpl::Connect(ETraceConnectType Type, const TCHAR* Parameter
 				UE_LOG_REF(LogCategory, Error, TEXT("Trace failed to connect (file: \"%s\")!"), Parameter ? Parameter : TEXT(""));
 			}
 
-			TraceType = FTraceAuxiliary::EConnectionType::File;
+			TraceType.store(FTraceAuxiliary::EConnectionType::File);
 		}
 
 		if (bConnected)
 		{
-			FTraceAuxiliary::OnTraceStarted.Broadcast(TraceType, TraceDest);
+			FTraceAuxiliary::OnTraceStarted.Broadcast(TraceType.load(), TraceDest);
 		}
 	}
 
@@ -337,10 +338,10 @@ bool FTraceAuxiliaryImpl::Stop()
 		return false;
 	}
 
-	FTraceAuxiliary::OnTraceStopped.Broadcast(TraceType, TraceDest);
+	FTraceAuxiliary::OnTraceStopped.Broadcast(TraceType.load(), TraceDest);
 
 	State.store(EState::Stopped);
-	TraceType = FTraceAuxiliary::EConnectionType::None;
+	TraceType.store(FTraceAuxiliary::EConnectionType::None);
 	TraceDest.Reset();
 	return true;
 }
@@ -564,6 +565,7 @@ bool FTraceAuxiliaryImpl::WriteSnapshot(const TCHAR* InFilePath, const FTraceAux
 
 	if (bResult)
 	{
+		FTraceAuxiliary::OnSnapshotSaved.Broadcast(NativePath);
 		UE_LOG_REF(LogCategory, Display, TEXT("Trace snapshot generated in %.3f seconds to \"%s\"."), FPlatformTime::Seconds() - StartTime, *NativePath);
 	}
 	else
@@ -586,6 +588,12 @@ bool FTraceAuxiliaryImpl::IsConnected() const
 	return State.load() == EState::Tracing;
 }
 
+////////////////////////////////////////////////////////////////////////////////
+FTraceAuxiliary::EConnectionType FTraceAuxiliaryImpl::GetConnectionType() const
+{
+	return TraceType.load();
+}
+
 ////////////////////////////////////////////////////////////////////////////////
 void FTraceAuxiliaryImpl::GetActiveChannelsString(FStringBuilderBase& String) const
 {
@@ -1265,6 +1273,7 @@ static bool StartFromCommandlineArguments(const TCHAR* CommandLine, bool& bOutSt
 FTraceAuxiliary::FOnConnection FTraceAuxiliary::OnConnection;
 FTraceAuxiliary::FOnTraceStarted FTraceAuxiliary::OnTraceStarted;
 FTraceAuxiliary::FOnTraceStopped FTraceAuxiliary::OnTraceStopped;
+FTraceAuxiliary::FOnSnapshotSaved FTraceAuxiliary::OnSnapshotSaved;
 
 ////////////////////////////////////////////////////////////////////////////////
 bool FTraceAuxiliary::Start(EConnectionType Type, const TCHAR* Target, const TCHAR* Channels, FOptions* Options, const FLogCategoryAlias& LogCategory)
@@ -1596,6 +1605,14 @@ bool FTraceAuxiliary::IsConnected()
 	return false;
 }
 
+FTraceAuxiliary::EConnectionType FTraceAuxiliary::GetConnectionType()
+{
+#if UE_TRACE_ENABLED
+	return GTraceAuxiliary.GetConnectionType();
+#endif
+	return FTraceAuxiliary::EConnectionType::None;
+}
+
 void FTraceAuxiliary::GetActiveChannelsString(FStringBuilderBase& String)
 {
 #if UE_TRACE_ENABLED
diff --git a/Engine/Source/Runtime/Core/Public/ProfilingDebugging/TraceAuxiliary.h b/Engine/Source/Runtime/Core/Public/ProfilingDebugging/TraceAuxiliary.h
--- a/Engine/Source/Runtime/Core/Public/ProfilingDebugging/TraceAuxiliary.h
+++ b/Engine/Source/Runtime/Core/Public/ProfilingDebugging/TraceAuxiliary.h
@@ -55,6 +55,12 @@ public:
 	*/
 	DECLARE_TS_MULTICAST_DELEGATE_TwoParams(FOnTraceStopped, FTraceAuxiliary::EConnectionType TraceType, const FString& TraceDestination);
 
+	/**
+	* Callback whenever a trace snapshot is saved.
+	* Path is the file system path of the snapshot file.
+	*/
+	DECLARE_TS_MULTICAST_DELEGATE_OneParam(FOnSnapshotSaved, const FString& Path);
+
 	struct FOptions
 	{
 		/** When set, trace will not start a worker thread, instead it is updated from end frame delegate. */
@@ -143,6 +149,11 @@ public:
 	 */
 	static bool	IsConnected();
 
+	/**
+	*  Returns the current connection type.
+	*/
+	static EConnectionType GetConnectionType();
+
 	/**
 	 *  Adds a comma separated list of currently active channels to the passed in StringBuilder
 	 */
@@ -174,5 +185,11 @@ public:
 	 * The type of recording and the destination (filepath or network) is passed to the delegate.
 	 */
 	static FOnTraceStopped OnTraceStopped;
+
+	/**
+	 * Delegate that triggers when a snapshot has been saved.
+	 * The path to the snapshot file is passed to the delegate.
+	 */
+	static FOnSnapshotSaved OnSnapshotSaved;
 	
 };
