Fix warning when creating the Insights Status Bar Widget.
Change order of Insights Status Bar Widget menu items.

diff --git a/Engine/Plugins/TraceUtilities/Source/EditorTraceUtilities/Private/InsightsStatusBar/SInsightsStatusBar.cpp b/Engine/Plugins/TraceUtilities/Source/EditorTraceUtilities/Private/InsightsStatusBar/SInsightsStatusBar.cpp
--- a/Engine/Plugins/TraceUtilities/Source/EditorTraceUtilities/Private/InsightsStatusBar/SInsightsStatusBar.cpp
+++ b/Engine/Plugins/TraceUtilities/Source/EditorTraceUtilities/Private/InsightsStatusBar/SInsightsStatusBar.cpp
@@ -45,7 +45,7 @@ void RegisterInsightsStatusWidgetWithToolMenu()
 {
 	UToolMenu* Menu = UToolMenus::Get()->ExtendMenu(TEXT("LevelEditor.StatusBar.ToolBar"));
 
-	FToolMenuSection& InsightsSection = Menu->AddSection(TEXT("Insights"), FText::GetEmpty(), FToolMenuInsert(TEXT("DDC"), EToolMenuInsertType::Before));
+	FToolMenuSection& InsightsSection = Menu->AddSection(TEXT("Insights"), FText::GetEmpty(), FToolMenuInsert(NAME_None, EToolMenuInsertType::First));
 
 	InsightsSection.AddEntry(
 		FToolMenuEntry::InitWidget(TEXT("InsightsStatusBar"), CreateInsightsStatusBarWidget(), FText::GetEmpty(), true, false)
@@ -224,20 +224,20 @@ TSharedRef<SWidget> SInsightsStatusBarWidget::MakeTraceMenu()
 	MenuBuilder.BeginSection("Tracing", LOCTEXT("TraceMenu_Section_Tracing", "Tracing"));
 	{
 		MenuBuilder.AddMenuEntry(
-			LOCTEXT("SaveSnapshotLabel", "Save Trace Snapshot"),
-			LOCTEXT("SaveSnapshotTooltip", "Save the current trace buffer to file."),
-			FSlateIcon(FEditorTraceUtilitiesStyle::Get().GetStyleSetName(), "Icons.TraceSnapshot"),
-			FUIAction(FExecuteAction::CreateSP(this, &SInsightsStatusBarWidget::SaveSnapshot),
-					  FCanExecuteAction::CreateSP(this, &SInsightsStatusBarWidget::SaveSnapshot_CanExecute)),
+			TAttribute<FText>::CreateSP(this, &SInsightsStatusBarWidget::GetTraceMenuItemText),
+			TAttribute<FText>::CreateSP(this, &SInsightsStatusBarWidget::GetTraceMenuItemTooltipText),
+			FSlateIcon(FEditorTraceUtilitiesStyle::Get().GetStyleSetName(), "Icons.StartTrace"),
+			FUIAction(FExecuteAction::CreateSP(this, &SInsightsStatusBarWidget::ToggleTracing_OnClicked)),
 			NAME_None,
 			EUserInterfaceActionType::Button
 		);
 
 		MenuBuilder.AddMenuEntry(
-			TAttribute<FText>::CreateSP(this, &SInsightsStatusBarWidget::GetTraceMenuItemText),
-			TAttribute<FText>::CreateSP(this, &SInsightsStatusBarWidget::GetTraceMenuItemTooltipText),
-			FSlateIcon(FEditorTraceUtilitiesStyle::Get().GetStyleSetName(), "Icons.StartTrace"),
-			FUIAction(FExecuteAction::CreateSP(this, &SInsightsStatusBarWidget::ToggleTracing_OnClicked)),
+			LOCTEXT("SaveSnapshotLabel", "Save Trace Snapshot"),
+			LOCTEXT("SaveSnapshotTooltip", "Save the current trace buffer to file."),
+			FSlateIcon(FEditorTraceUtilitiesStyle::Get().GetStyleSetName(), "Icons.TraceSnapshot"),
+			FUIAction(FExecuteAction::CreateSP(this, &SInsightsStatusBarWidget::SaveSnapshot),
+					  FCanExecuteAction::CreateSP(this, &SInsightsStatusBarWidget::SaveSnapshot_CanExecute)),
 			NAME_None,
 			EUserInterfaceActionType::Button
 		);
