Organize the Insights Status Bar icons in a better way.

diff --git a/Engine/Plugins/TraceUtilities/Source/EditorTraceUtilities/Private/InsightsStatusBar/EditorTraceUtilitiesStyle.h b/Engine/Plugins/TraceUtilities/Source/EditorTraceUtilities/Private/InsightsStatusBar/EditorTraceUtilitiesStyle.h
--- a/Engine/Plugins/TraceUtilities/Source/EditorTraceUtilities/Private/InsightsStatusBar/EditorTraceUtilitiesStyle.h
+++ b/Engine/Plugins/TraceUtilities/Source/EditorTraceUtilities/Private/InsightsStatusBar/EditorTraceUtilitiesStyle.h
@@ -24,15 +24,24 @@ public:
 		SetContentRoot(FPaths::EnginePluginsDir() / TEXT("/TraceUtilities/Content/Slate"));
 		SetCoreContentRoot(FPaths::EngineContentDir() / TEXT("Slate"));
 
-		Set("Icons.OpenLiveSession", new CORE_IMAGE_BRUSH_SVG("Starship/Insights/Session", Icon16x16));
-
-		Set("Icons.Trace", new IMAGE_BRUSH_SVG("Trace", Icon16x16));
-		Set("Icons.RecordTrace", new IMAGE_BRUSH_SVG("RecordTrace", Icon16x16));
-		Set("Icons.RecordTraceOutline", new IMAGE_BRUSH_SVG("RecordTraceOutline", Icon16x16));
-		Set("Icons.RecordTraceWithOutline", new IMAGE_BRUSH_SVG("RecordTraceWithOutline", Icon16x16));
-		Set("Icons.RecordTraceStop", new IMAGE_BRUSH_SVG("RecordTraceStop", Icon16x16, FStyleColors::Error));
-		Set("Icons.StartTrace", new IMAGE_BRUSH_SVG("StartTrace", Icon16x16));
-		Set("Icons.TraceSnapshot", new IMAGE_BRUSH_SVG("TraceSnapshot", Icon16x16));
+		Set("Icons.OpenLiveSession.Menu", new CORE_IMAGE_BRUSH_SVG("Starship/Insights/Session", Icon16x16));
+
+		Set("Icons.UnrealInsights.Menu", new IMAGE_BRUSH_SVG("UnrealInsights_16", Icon16x16));
+
+		Set("Icons.Trace.Menu", new IMAGE_BRUSH_SVG("Trace_16", Icon16x16));
+		Set("Icons.Trace.StatusBar", new IMAGE_BRUSH_SVG("Trace_20", Icon16x16));
+
+		Set("Icons.RecordTraceCenter.StatusBar", new IMAGE_BRUSH_SVG("RecordTraceCenter_20", Icon16x16));
+		Set("Icons.RecordTraceOutline.StatusBar", new IMAGE_BRUSH_SVG("RecordTraceOutline_20", Icon16x16));
+		Set("Icons.RecordTraceRecording.StatusBar", new IMAGE_BRUSH_SVG("RecordTraceRecording_20", Icon16x16));
+		
+		Set("Icons.RecordTraceStop.StatusBar", new IMAGE_BRUSH_SVG("RecordTraceStop_20", Icon16x16, FStyleColors::Error));
+
+		Set("Icons.StartTrace.Menu", new IMAGE_BRUSH_SVG("StartTrace_16", Icon16x16));
+		Set("Icons.StartTrace.StatusBar", new IMAGE_BRUSH_SVG("StartTrace_20", Icon16x16));
+
+		Set("Icons.TraceSnapshot.Menu", new IMAGE_BRUSH_SVG("TraceSnapshot_16", Icon16x16));
+		Set("Icons.TraceSnapshot.StatusBar", new IMAGE_BRUSH_SVG("TraceSnapshot_20", Icon16x16));
 
 		FSlateStyleRegistry::RegisterSlateStyle(*this);
 	}
diff --git a/Engine/Plugins/TraceUtilities/Source/EditorTraceUtilities/Private/InsightsStatusBar/SInsightsStatusBar.cpp b/Engine/Plugins/TraceUtilities/Source/EditorTraceUtilities/Private/InsightsStatusBar/SInsightsStatusBar.cpp
--- a/Engine/Plugins/TraceUtilities/Source/EditorTraceUtilities/Private/InsightsStatusBar/SInsightsStatusBar.cpp
+++ b/Engine/Plugins/TraceUtilities/Source/EditorTraceUtilities/Private/InsightsStatusBar/SInsightsStatusBar.cpp
@@ -77,7 +77,7 @@ void SInsightsStatusBarWidget::Construct(const FArguments& InArgs)
 				[
 					SNew(SImage)
 					.ColorAndOpacity(FSlateColor::UseForeground())
-					.Image(FEditorTraceUtilitiesStyle::Get().GetBrush("Icons.Trace"))
+					.Image(FEditorTraceUtilitiesStyle::Get().GetBrush("Icons.Trace.StatusBar"))
 				]
 
 				+ SHorizontalBox::Slot()
@@ -111,7 +111,7 @@ void SInsightsStatusBarWidget::Construct(const FArguments& InArgs)
 				[
 					SNew(SImage)
 					.ColorAndOpacity(this, &SInsightsStatusBarWidget::GetRecordingButtonColor)
-					.Image(FEditorTraceUtilitiesStyle::Get().GetBrush("Icons.RecordTrace"))
+					.Image(FEditorTraceUtilitiesStyle::Get().GetBrush("Icons.RecordTraceCenter.StatusBar"))
 					.Visibility(this, &SInsightsStatusBarWidget::GetStartTraceIconVisibility)
 				]
 
@@ -119,14 +119,14 @@ void SInsightsStatusBarWidget::Construct(const FArguments& InArgs)
 				[
 					SNew(SImage)
 					.ColorAndOpacity(this, &SInsightsStatusBarWidget::GetRecordingButtonOutlineColor)
-					.Image(FEditorTraceUtilitiesStyle::Get().GetBrush("Icons.RecordTraceOutline"))
+					.Image(FEditorTraceUtilitiesStyle::Get().GetBrush("Icons.RecordTraceOutline.StatusBar"))
 					.Visibility(this, &SInsightsStatusBarWidget::GetStartTraceIconVisibility)
 				]
 
 				+ SOverlay::Slot()
 				[
 					SNew(SImage)
-					.Image(FEditorTraceUtilitiesStyle::Get().GetBrush("Icons.RecordTraceStop"))
+					.Image(FEditorTraceUtilitiesStyle::Get().GetBrush("Icons.RecordTraceStop.StatusBar"))
 					.Visibility(this, &SInsightsStatusBarWidget::GetStopTraceIconVisibility)
 				]
 			]
@@ -145,7 +145,7 @@ void SInsightsStatusBarWidget::Construct(const FArguments& InArgs)
 			[
 				SNew(SImage)
 				.DesiredSizeOverride(CoreStyleConstants::Icon16x16)
-				.Image(FEditorTraceUtilitiesStyle::Get().GetBrush("Icons.TraceSnapshot"))
+				.Image(FEditorTraceUtilitiesStyle::Get().GetBrush("Icons.TraceSnapshot.StatusBar"))
 				.ToolTipText(LOCTEXT("SaveSnapShot","Snapshot: Save Current Trace Buffer To File"))
 			]
 		]
@@ -190,7 +190,7 @@ TSharedRef<SWidget> SInsightsStatusBarWidget::MakeTraceMenu()
 		//		LOCTEXT("Channels_Desc", "Select what trace channels to enable when tracing."),
 		//		FNewMenuDelegate::CreateSP(this, &SInsightsStatusBarWidget::Channels_BuildMenu),
 		//		false,
-		//		FSlateIcon(FEditorTraceUtilitiesStyle::Get().GetStyleSetName(), ("Icons.Trace"))
+		//		FSlateIcon(FEditorTraceUtilitiesStyle::Get().GetStyleSetName(), ("Icons.Trace.Menu"))
 		//	);
 	}
 	MenuBuilder.EndSection();
@@ -226,7 +226,7 @@ TSharedRef<SWidget> SInsightsStatusBarWidget::MakeTraceMenu()
 		MenuBuilder.AddMenuEntry(
 			TAttribute<FText>::CreateSP(this, &SInsightsStatusBarWidget::GetTraceMenuItemText),
 			TAttribute<FText>::CreateSP(this, &SInsightsStatusBarWidget::GetTraceMenuItemTooltipText),
-			FSlateIcon(FEditorTraceUtilitiesStyle::Get().GetStyleSetName(), "Icons.StartTrace"),
+			FSlateIcon(FEditorTraceUtilitiesStyle::Get().GetStyleSetName(), "Icons.StartTrace.Menu"),
 			FUIAction(FExecuteAction::CreateSP(this, &SInsightsStatusBarWidget::ToggleTracing_OnClicked)),
 			NAME_None,
 			EUserInterfaceActionType::Button
@@ -235,7 +235,7 @@ TSharedRef<SWidget> SInsightsStatusBarWidget::MakeTraceMenu()
 		MenuBuilder.AddMenuEntry(
 			LOCTEXT("SaveSnapshotLabel", "Save Trace Snapshot"),
 			LOCTEXT("SaveSnapshotTooltip", "Save the current trace buffer to file."),
-			FSlateIcon(FEditorTraceUtilitiesStyle::Get().GetStyleSetName(), "Icons.TraceSnapshot"),
+			FSlateIcon(FEditorTraceUtilitiesStyle::Get().GetStyleSetName(), "Icons.TraceSnapshot.Menu"),
 			FUIAction(FExecuteAction::CreateSP(this, &SInsightsStatusBarWidget::SaveSnapshot),
 					  FCanExecuteAction::CreateSP(this, &SInsightsStatusBarWidget::SaveSnapshot_CanExecute)),
 			NAME_None,
@@ -308,7 +308,7 @@ TSharedRef<SWidget> SInsightsStatusBarWidget::MakeTraceMenu()
 		MenuBuilder.AddMenuEntry(
 			LOCTEXT("OpenInsightsLabel", "Unreal Insights (Session Browser)"),
 			LOCTEXT("OpenInsightsTooltip", "Launch the Unreal Insights Session Browser."),
-			FSlateIcon(FAppStyle::Get().GetStyleSetName(), "UnrealInsights.MenuIcon"),
+			FSlateIcon(FEditorTraceUtilitiesStyle::Get().GetStyleSetName(), "Icons.UnrealInsights.Menu"),
 			FUIAction(FExecuteAction::CreateSP(this, &SInsightsStatusBarWidget::LaunchUnrealInsights_OnClicked)),
 			NAME_None,
 			EUserInterfaceActionType::Button
@@ -317,7 +317,7 @@ TSharedRef<SWidget> SInsightsStatusBarWidget::MakeTraceMenu()
 		MenuBuilder.AddMenuEntry(
 			LOCTEXT("OpenLiveSessionLabel", "Open Live Session"),
 			LOCTEXT("OpenLiveSessionTooltip", "Opening the live session is possible only while tracing to the trace store."),
-			FSlateIcon(FEditorTraceUtilitiesStyle::Get().GetStyleSetName(), "Icons.OpenLiveSession"),
+			FSlateIcon(FEditorTraceUtilitiesStyle::Get().GetStyleSetName(), "Icons.OpenLiveSession.Menu"),
 			FUIAction(FExecuteAction::CreateSP(this, &SInsightsStatusBarWidget::OpenLiveSession_OnClicked),
 				FCanExecuteAction::CreateLambda([this]() { return FTraceAuxiliary::GetConnectionType() == FTraceAuxiliary::EConnectionType::Network; })),
 			NAME_None,
