Show the recent traces in the Insights Status Bar widget.

diff --git a/Engine/Plugins/TraceUtilities/Source/EditorTraceUtilities/EditorTraceUtilities.Build.cs b/Engine/Plugins/TraceUtilities/Source/EditorTraceUtilities/EditorTraceUtilities.Build.cs
--- a/Engine/Plugins/TraceUtilities/Source/EditorTraceUtilities/EditorTraceUtilities.Build.cs
+++ b/Engine/Plugins/TraceUtilities/Source/EditorTraceUtilities/EditorTraceUtilities.Build.cs
@@ -27,7 +27,8 @@ public class EditorTraceUtilities : ModuleRules
             {
                 "CoreUObject",
                 "Engine",
-                "Slate",
+				"InputCore",
+				"Slate",
                 "SlateCore",
                 "ToolMenus",
                 "ToolWidgets",
diff --git a/Engine/Plugins/TraceUtilities/Source/EditorTraceUtilities/Private/InsightsStatusBar/EditorTraceUtilitiesStyle.h b/Engine/Plugins/TraceUtilities/Source/EditorTraceUtilities/Private/InsightsStatusBar/EditorTraceUtilitiesStyle.h
--- a/Engine/Plugins/TraceUtilities/Source/EditorTraceUtilities/Private/InsightsStatusBar/EditorTraceUtilitiesStyle.h
+++ b/Engine/Plugins/TraceUtilities/Source/EditorTraceUtilities/Private/InsightsStatusBar/EditorTraceUtilitiesStyle.h
@@ -43,6 +43,9 @@ public:
 		Set("Icons.TraceSnapshot.Menu", new IMAGE_BRUSH_SVG("TraceSnapshot_16", Icon16x16));
 		Set("Icons.TraceSnapshot.StatusBar", new IMAGE_BRUSH_SVG("TraceSnapshot_20", Icon16x16));
 
+		Set("Icons.TraceStore.Menu", new CORE_IMAGE_BRUSH_SVG("Starship/Insights/TraceStore", Icon16x16));
+		Set("Icons.File.Menu", new CORE_IMAGE_BRUSH_SVG("Starship/Common/file", Icon16x16));
+
 		FSlateStyleRegistry::RegisterSlateStyle(*this);
 	}
 
diff --git a/Engine/Plugins/TraceUtilities/Source/EditorTraceUtilities/Private/InsightsStatusBar/SInsightsStatusBar.cpp b/Engine/Plugins/TraceUtilities/Source/EditorTraceUtilities/Private/InsightsStatusBar/SInsightsStatusBar.cpp
--- a/Engine/Plugins/TraceUtilities/Source/EditorTraceUtilities/Private/InsightsStatusBar/SInsightsStatusBar.cpp
+++ b/Engine/Plugins/TraceUtilities/Source/EditorTraceUtilities/Private/InsightsStatusBar/SInsightsStatusBar.cpp
@@ -12,7 +12,9 @@
 #include "HAL/FileManager.h"
 #include "Math/Color.h"
 #include "Misc/ConfigCacheIni.h"
+#include "Misc/Paths.h"
 #include "Modules/ModuleManager.h"
+#include "SRecentTracesList.h"
 #include "Styling/StyleColors.h"
 #include "ToolMenus.h"
 #include "Trace/StoreClient.h"
@@ -323,6 +325,15 @@ TSharedRef<SWidget> SInsightsStatusBarWidget::MakeTraceMenu()
 			NAME_None,
 			EUserInterfaceActionType::Button
 		);
+
+		MenuBuilder.AddSubMenu
+		(
+			LOCTEXT("RecentTracesLabel", "Recent Traces"),
+			LOCTEXT("RecentTracesTooltop", "Open the latest traces recorded to the trace store or as a file."),
+			FNewMenuDelegate::CreateSP(this, &SInsightsStatusBarWidget::Traces_BuildMenu),
+			false,
+			FSlateIcon(FEditorTraceUtilitiesStyle::Get().GetStyleSetName(), ("Icons.Trace.Menu"))
+		);
 	}
 	MenuBuilder.EndSection();
 
@@ -380,6 +391,21 @@ void SInsightsStatusBarWidget::Channels_BuildMenu(FMenuBuilder& MenuBuilder)
 	MenuBuilder.EndSection();
 }
 
+void SInsightsStatusBarWidget::Traces_BuildMenu(FMenuBuilder& MenuBuilder)
+{
+	CacheTraceStorePath();
+
+	MenuBuilder.AddWidget(
+		SNew(SBox)
+		.MaxDesiredHeight(400.0f)
+		.MinDesiredWidth(200.0f)
+		.MaxDesiredWidth(270.0f)
+		[
+			SNew(SRecentTracesList, TraceStorePath)
+		],
+		FText(), true);
+}
+
 FText SInsightsStatusBarWidget::GetTitleToolTipText() const
 {
 	FTextBuilder DescBuilder;
@@ -491,23 +517,9 @@ void SInsightsStatusBarWidget::OpenTraceStoreDirectory_OnClicked()
 
 void SInsightsStatusBarWidget::OpenTraceStoreDirectory(bool bSelectLastTrace)
 {
-	UE::Trace::FStoreClient* StoreClient = UE::Trace::FStoreClient::Connect(TEXT("localhost"));
-
-	if (!StoreClient)
-	{
-		// TODO: Add Error Message
-		return;
-	}
-
-	const UE::Trace::FStoreClient::FStatus* Status = StoreClient->GetStatus();
-	if (!Status)
-	{
-		// TODO: Add Error Message
-		return;
-	}
-
-	FString Path(Status->GetStoreDir());
+	CacheTraceStorePath();
 
+	FString StorePathLocal = TraceStorePath;
 	if (bSelectLastTrace)
 	{
 		FString MostRecentTraceName;
@@ -526,15 +538,15 @@ void SInsightsStatusBarWidget::OpenTraceStoreDirectory(bool bSelectLastTrace)
 			return true;
 		};
 
-		IFileManager::Get().IterateDirectoryStat(*Path, Visitor);
+		IFileManager::Get().IterateDirectoryStat(*StorePathLocal, Visitor);
 
 		if (!MostRecentTraceName.IsEmpty())
 		{
-			Path = MostRecentTraceName;
+			StorePathLocal = MostRecentTraceName;
 		}
 	}
 
-	FString FullPath(FPaths::ConvertRelativePathToFull(Path));
+	FString FullPath(FPaths::ConvertRelativePathToFull(StorePathLocal));
 	FPlatformProcess::ExploreFolder(*FullPath);
 }
 
@@ -697,4 +709,27 @@ void SInsightsStatusBarWidget::OnSnapshotSaved(const FString& InPath)
 	}
 }
 
+void SInsightsStatusBarWidget::CacheTraceStorePath()
+{
+	if (TraceStorePath.IsEmpty())
+	{
+		UE::Trace::FStoreClient* StoreClient = UE::Trace::FStoreClient::Connect(TEXT("localhost"));
+
+		if (!StoreClient)
+		{
+			// TODO: Add Error Message
+			return;
+		}
+
+		const UE::Trace::FStoreClient::FStatus* Status = StoreClient->GetStatus();
+		if (!Status)
+		{
+			// TODO: Add Error Message
+			return;
+		}
+
+		TraceStorePath = FString(Status->GetStoreDir());
+	}
+}
+
 #undef LOCTEXT_NAMESPACE
diff --git a/Engine/Plugins/TraceUtilities/Source/EditorTraceUtilities/Private/InsightsStatusBar/SInsightsStatusBar.h b/Engine/Plugins/TraceUtilities/Source/EditorTraceUtilities/Private/InsightsStatusBar/SInsightsStatusBar.h
--- a/Engine/Plugins/TraceUtilities/Source/EditorTraceUtilities/Private/InsightsStatusBar/SInsightsStatusBar.h
+++ b/Engine/Plugins/TraceUtilities/Source/EditorTraceUtilities/Private/InsightsStatusBar/SInsightsStatusBar.h
@@ -66,6 +66,7 @@ private:
 
 	TSharedRef<SWidget> MakeTraceMenu();
 	void Channels_BuildMenu(FMenuBuilder& MenuBuilder);
+	void Traces_BuildMenu(FMenuBuilder& MenuBuilder);
 
 	static void SendSnapshotNotification();
 	static void SendTraceStartedNotification();
@@ -80,6 +81,8 @@ private:
 	void OnTraceStopped(FTraceAuxiliary::EConnectionType TraceType, const FString& TraceDestination);
 	void OnSnapshotSaved(const FString& InPath);
 
+	void CacheTraceStorePath();
+
 private:
 	static const TCHAR* DefaultPreset;
 	static const TCHAR* MemoryPreset;
@@ -95,4 +98,6 @@ private:
 	bool bIsTraceRecordButtonHovered = false;
 	const TCHAR* Channels;
 	mutable double ConnectionStartTime = 0.0f;
+
+	FString TraceStorePath;
 };
diff --git a/Engine/Plugins/TraceUtilities/Source/EditorTraceUtilities/Private/InsightsStatusBar/SRecentTracesList.cpp b/Engine/Plugins/TraceUtilities/Source/EditorTraceUtilities/Private/InsightsStatusBar/SRecentTracesList.cpp
new file mode 100644
--- /dev/null
+++ b/Engine/Plugins/TraceUtilities/Source/EditorTraceUtilities/Private/InsightsStatusBar/SRecentTracesList.cpp
@@ -0,0 +1,213 @@
+// Copyright Epic Games, Inc. All Rights Reserved.
+
+#include "SRecentTracesList.h"
+
+#include "EditorTraceUtilitiesStyle.h"
+#include "Fonts/FontMeasure.h"
+#include "Framework/Application/SlateApplication.h"
+#include "GenericPlatform/GenericPlatformFile.h"
+#include "HAL/FileManager.h"
+#include "Misc/Paths.h"
+#include "UnrealInsightsLauncher.h"
+#include "Widgets/Input/SButton.h"
+#include "Widgets/Layout/SScrollBorder.h"
+
+#define LOCTEXT_NAMESPACE "RecentTracesList"
+
+class SRecentTracesListEntry : public STableRow<TSharedPtr<FString>>
+{
+public:
+	SLATE_BEGIN_ARGS(SRecentTracesListEntry) {}
+
+	SLATE_ARGUMENT(TSharedPtr<FTraceFileInfo>, TraceInfo)
+
+	SLATE_END_ARGS()
+
+	void Construct(const FArguments& InArgs, const TSharedRef<STableViewBase>& InOwnerTable)
+	{
+		TraceInfo = InArgs._TraceInfo;
+
+		STableRow<TSharedPtr<FString>>::Construct(STableRow<TSharedPtr<FString>>::FArguments(), InOwnerTable);
+	}
+
+	virtual void ConstructChildren(ETableViewMode::Type InOwnerTableMode, const TAttribute<FMargin>& InPadding, const TSharedRef<SWidget>& InContent) override
+	{
+		FString TraceName = FPaths::GetBaseFilename(TraceInfo->FilePath);
+		TSharedRef<FSlateFontMeasure> FontMeasureService = FSlateApplication::Get().GetRenderer()->GetFontMeasureService();
+		int32 LastCharacter = FontMeasureService->FindLastWholeCharacterIndexBeforeOffset(TraceName, FAppStyle::GetFontStyle("Menu.Label"), 180);
+		if (LastCharacter < TraceName.Len() - 1 && LastCharacter > 3)
+		{
+			TraceName.LeftInline(LastCharacter - 3);
+			TraceName += TEXT("...");
+		}
+
+		FString TooltipText = TraceInfo->FilePath;
+		FPaths::NormalizeFilename(TooltipText);
+		TooltipText = FPaths::ConvertRelativePathToFull(TooltipText);
+
+		const FSlateBrush* TraceLocationIcon;
+		FText TraceLocationTooltip;
+
+		if (TraceInfo->bIsFromTraceStore)
+		{
+			TraceLocationIcon = FEditorTraceUtilitiesStyle::Get().GetBrush("Icons.TraceStore.Menu");
+			TraceLocationTooltip = LOCTEXT("TraceStoreLocationTooltip", "This trace is located in the trace store.");
+		}
+		else
+		{
+			TraceLocationIcon = FEditorTraceUtilitiesStyle::Get().GetBrush("Icons.File.Menu");
+			TraceLocationTooltip = LOCTEXT("TraceStoreLocationTooltip", "This trace was saved to file and is located in the current project's profilling folder.");
+		}
+
+		ChildSlot
+			[
+				SNew(SBorder)
+				.BorderImage(FAppStyle::GetBrush("Menu.Background"))
+				[
+					SNew(SHorizontalBox)
+
+					+ SHorizontalBox::Slot()
+					.VAlign(EVerticalAlignment::VAlign_Center)
+					.HAlign(EHorizontalAlignment::HAlign_Left)
+					.Padding(8.0f, 0.0f, 0.0f, 0.0f)
+					.AutoWidth()
+					[
+						SNew(SImage)
+						.ColorAndOpacity(FSlateColor::UseSubduedForeground())
+						.Image(TraceLocationIcon)
+						.ToolTipText(TraceLocationTooltip)
+					]
+
+					+ SHorizontalBox::Slot()
+					.VAlign(EVerticalAlignment::VAlign_Center)
+					.HAlign(EHorizontalAlignment::HAlign_Left)
+					.Padding(0.0f, 2.0f, 0.0f, 0.0f)
+					.AutoWidth()
+					[
+						SNew(STextBlock)
+						.TextStyle(FAppStyle::Get(), "Menu.Label")
+						.Text(FText::FromString(TraceName))
+						.ToolTipText(FText::FromString(TooltipText))
+						.ColorAndOpacity(FSlateColor::UseForeground())
+					]
+					
+					+ SHorizontalBox::Slot()
+					.FillWidth(1.0f)
+					.HAlign(EHorizontalAlignment::HAlign_Right)
+					[
+						SNew(SHorizontalBox)
+
+						+ SHorizontalBox::Slot()
+						.HAlign(EHorizontalAlignment::HAlign_Right)
+						.FillWidth(1.0f)
+						[
+							SNew(SButton)
+							.ButtonStyle(FAppStyle::Get(), "SimpleButton")
+							.OnClicked(this, &SRecentTracesListEntry::OpenContainingFolder)
+							.ToolTipText(LOCTEXT("ExploreFolderTooltip", "Open the folder containing the trace file."))
+							.Content()
+							[
+								SNew(SImage)
+								.ColorAndOpacity(FSlateColor::UseSubduedForeground())
+								.Image(FAppStyle::Get().GetBrush("Icons.FolderOpen"))
+							]
+						]
+
+						+ SHorizontalBox::Slot()
+						.HAlign(EHorizontalAlignment::HAlign_Right)
+						.AutoWidth()
+						[
+							SNew(SButton)
+							.ButtonStyle(FAppStyle::Get(), "SimpleButton")
+							.OnClicked(this, &SRecentTracesListEntry::OpenTrace)
+							.ToolTipText(LOCTEXT("OpenTraceTooltip", "Open the trace file in Unreal Insights."))
+							.Content()
+							[
+								SNew(SImage)
+								.ColorAndOpacity(FSlateColor::UseSubduedForeground())
+								.Image(FEditorTraceUtilitiesStyle::Get().GetBrush("Icons.UnrealInsights.Menu"))
+							]
+						]
+					]
+				]
+			];
+	}
+
+	FReply OpenContainingFolder()
+	{
+		FPlatformProcess::ExploreFolder(*TraceInfo->FilePath);
+
+		return FReply::Handled();
+	}
+
+	FReply OpenTrace()
+	{
+		FUnrealInsightsLauncher::Get()->TryOpenTraceFromDestination(TraceInfo->FilePath);
+		return FReply::Handled();
+	}
+
+	TSharedPtr<FTraceFileInfo> TraceInfo;
+};
+
+void SRecentTracesList::Construct(const FArguments& InArgs, const FString& InStorePath)
+{
+	StorePath = InStorePath;
+
+	PopulateRecentTracesList();
+
+	ListView = SNew(SListView<TSharedPtr<FTraceFileInfo>>)
+		.IsFocusable(true)
+		.SelectionMode(ESelectionMode::None)
+		.ListItemsSource(&Traces)
+		.OnGenerateRow(this, &SRecentTracesList::OnGenerateRow);
+	
+	ChildSlot
+	[
+		SNew(SVerticalBox)
+
+		+ SVerticalBox::Slot()
+		.Padding(FMargin(2.0f, 0.0f))
+		.FillHeight(1.0f)
+		[
+			SNew(SScrollBorder, ListView.ToSharedRef())
+			[
+				ListView.ToSharedRef()
+			]
+		]
+	];
+}
+
+TSharedRef<ITableRow> SRecentTracesList::OnGenerateRow(TSharedPtr<FTraceFileInfo> Item, const TSharedRef<STableViewBase>& OwnerTable)
+{
+	return SNew(SRecentTracesListEntry, OwnerTable)
+		.TraceInfo(Item);
+}
+
+void SRecentTracesList::PopulateRecentTracesList()
+{
+	Traces.Empty();
+	bool bIsFromTraceStore = true;
+	auto Visitor = [this, &bIsFromTraceStore](const TCHAR* Filename, const FFileStatData& StatData)
+	{
+		if (FPaths::GetExtension(Filename) == TEXT("utrace"))
+		{
+			TSharedPtr<FTraceFileInfo> TraceInfo = MakeShared<FTraceFileInfo>();
+			TraceInfo->FilePath = Filename;
+			TraceInfo->ModifiedTime = StatData.ModificationTime;
+			TraceInfo->bIsFromTraceStore = bIsFromTraceStore;
+
+			Traces.Add(TraceInfo);
+		}
+
+		return true;
+	};
+
+	IFileManager::Get().IterateDirectoryStat(*StorePath, Visitor);
+	
+	bIsFromTraceStore = false;
+	IFileManager::Get().IterateDirectoryStat(*FPaths::ProfilingDir(), Visitor);
+	
+	Algo::SortBy(Traces, [](TSharedPtr<FTraceFileInfo> TraceInfo) { return *TraceInfo; });
+}
+
+#undef LOCTEXT_NAMESPACE
diff --git a/Engine/Plugins/TraceUtilities/Source/EditorTraceUtilities/Private/InsightsStatusBar/SRecentTracesList.h b/Engine/Plugins/TraceUtilities/Source/EditorTraceUtilities/Private/InsightsStatusBar/SRecentTracesList.h
new file mode 100644
--- /dev/null
+++ b/Engine/Plugins/TraceUtilities/Source/EditorTraceUtilities/Private/InsightsStatusBar/SRecentTracesList.h
@@ -0,0 +1,41 @@
+// Copyright Epic Games, Inc. All Rights Reserved.
+
+#pragma once
+
+#include "Widgets/SCompoundWidget.h"
+#include "Widgets/Views/SListView.h"
+
+struct FTraceFileInfo
+{
+	FString FilePath;
+	FDateTime ModifiedTime;
+	bool bIsFromTraceStore;
+
+	bool operator <(const FTraceFileInfo& rhs)
+	{
+		return this->ModifiedTime > rhs.ModifiedTime;
+	}
+};
+
+/**
+ *  Shows a list of recent traces in a menu.
+ */
+class SRecentTracesList : public SCompoundWidget
+{
+	SLATE_BEGIN_ARGS(SRecentTracesList) {}
+	SLATE_END_ARGS()
+
+	void Construct(const FArguments& InArgs, const FString& InStorePath);
+
+	TSharedRef<ITableRow> OnGenerateRow(TSharedPtr<FTraceFileInfo> Item, const TSharedRef<STableViewBase>& OwnerTable);
+
+private:
+	void PopulateRecentTracesList();
+
+private:
+	TSharedPtr<SListView<TSharedPtr<FTraceFileInfo>>> ListView;
+	
+	TArray<TSharedPtr<FTraceFileInfo>> Traces;
+
+	FString StorePath;
+};
