Fix duplicate LOCTEXT key.

diff --git a/Engine/Plugins/TraceUtilities/Source/EditorTraceUtilities/Private/InsightsStatusBar/SRecentTracesList.cpp b/Engine/Plugins/TraceUtilities/Source/EditorTraceUtilities/Private/InsightsStatusBar/SRecentTracesList.cpp
--- a/Engine/Plugins/TraceUtilities/Source/EditorTraceUtilities/Private/InsightsStatusBar/SRecentTracesList.cpp
+++ b/Engine/Plugins/TraceUtilities/Source/EditorTraceUtilities/Private/InsightsStatusBar/SRecentTracesList.cpp
@@ -56,7 +56,7 @@ public:
 		else
 		{
 			TraceLocationIcon = FEditorTraceUtilitiesStyle::Get().GetBrush("Icons.File.Menu");
-			TraceLocationTooltip = LOCTEXT("TraceStoreLocationTooltip", "This trace was saved to file and is located in the current project's profilling folder.");
+			TraceLocationTooltip = LOCTEXT("TraceFileLocationTooltip", "This trace was saved to file and is located in the current project's profilling folder.");
 		}
 
 		ChildSlot
