Trace channels can now be controlled from the Insights Status Bar Widget.

diff --git a/Engine/Plugins/TraceUtilities/Source/EditorTraceUtilities/EditorTraceUtilities.Build.cs b/Engine/Plugins/TraceUtilities/Source/EditorTraceUtilities/EditorTraceUtilities.Build.cs
--- a/Engine/Plugins/TraceUtilities/Source/EditorTraceUtilities/EditorTraceUtilities.Build.cs
+++ b/Engine/Plugins/TraceUtilities/Source/EditorTraceUtilities/EditorTraceUtilities.Build.cs
@@ -36,8 +36,9 @@ public class EditorTraceUtilities : ModuleRules
                 "ApplicationCore",
                 "UATHelper",
                 "TraceInsights",
-                "TraceAnalysis"
-            }
+                "TraceAnalysis",
+				"TraceLog",
+			}
         );
     }
 }
\ No newline at end of file
diff --git a/Engine/Plugins/TraceUtilities/Source/EditorTraceUtilities/Private/InsightsStatusBar/SInsightsStatusBar.cpp b/Engine/Plugins/TraceUtilities/Source/EditorTraceUtilities/Private/InsightsStatusBar/SInsightsStatusBar.cpp
--- a/Engine/Plugins/TraceUtilities/Source/EditorTraceUtilities/Private/InsightsStatusBar/SInsightsStatusBar.cpp
+++ b/Engine/Plugins/TraceUtilities/Source/EditorTraceUtilities/Private/InsightsStatusBar/SInsightsStatusBar.cpp
@@ -54,6 +54,13 @@ void RegisterInsightsStatusWidgetWithToolMenu()
 	);
 }
 
+void SInsightsStatusBarWidget::Tick(const FGeometry& AllottedGeometry, const double InCurrentTime, const float InDeltaTime)
+{
+	SCompoundWidget::Tick(AllottedGeometry, InCurrentTime, InDeltaTime);
+
+	bShouldUpdateChannels = true;
+}
+
 void SInsightsStatusBarWidget::Construct(const FArguments& InArgs)
 {
 	this->ChildSlot
@@ -162,7 +169,6 @@ void SInsightsStatusBarWidget::Construct(const FArguments& InArgs)
 		TraceDestination = ETraceDestination::File;
 	}
 
-	Channels = DefaultPreset;
 	FTraceAuxiliary::OnTraceStarted.AddSP(this, &SInsightsStatusBarWidget::OnTraceStarted);
 	FTraceAuxiliary::OnTraceStopped.AddSP(this, &SInsightsStatusBarWidget::OnTraceStopped);
 	FTraceAuxiliary::OnSnapshotSaved.AddSP(this, &SInsightsStatusBarWidget::OnSnapshotSaved);
@@ -185,15 +191,14 @@ TSharedRef<SWidget> SInsightsStatusBarWidget::MakeTraceMenu()
 			EUserInterfaceActionType::ToggleButton
 		);
 
-		// To be activated 
-		//	MenuBuilder.AddSubMenu
-		//	(
-		//		LOCTEXT("Channels", "Channels"),
-		//		LOCTEXT("Channels_Desc", "Select what trace channels to enable when tracing."),
-		//		FNewMenuDelegate::CreateSP(this, &SInsightsStatusBarWidget::Channels_BuildMenu),
-		//		false,
-		//		FSlateIcon(FEditorTraceUtilitiesStyle::Get().GetStyleSetName(), ("Icons.Trace.Menu"))
-		//	);
+		MenuBuilder.AddSubMenu
+		(
+			LOCTEXT("Channels", "Channels"),
+			LOCTEXT("Channels_Desc", "Select what trace channels to enable when tracing."),
+			FNewMenuDelegate::CreateSP(this, &SInsightsStatusBarWidget::Channels_BuildMenu),
+			false,
+			FSlateIcon(FEditorTraceUtilitiesStyle::Get().GetStyleSetName(), ("Icons.Trace.Menu"))
+		);
 	}
 	MenuBuilder.EndSection();
 
@@ -342,51 +347,26 @@ TSharedRef<SWidget> SInsightsStatusBarWidget::MakeTraceMenu()
 
 void SInsightsStatusBarWidget::Channels_BuildMenu(FMenuBuilder& MenuBuilder)
 {
-	MenuBuilder.BeginSection("Presets", LOCTEXT("Presets", "Presets"));
-	{
-		MenuBuilder.AddMenuEntry(
-			LOCTEXT("DefautLabel", "Default"),
-			LOCTEXT("DefaultLabelDesc", "Activate the cpu,gpu,frame,log,bookmark channels."),
-			FSlateIcon(),
-			FUIAction(FExecuteAction::CreateSP(this, &SInsightsStatusBarWidget::SetTraceChannels, DefaultPreset),
-				FCanExecuteAction::CreateLambda([]() {return true; }),
-				FIsActionChecked::CreateSP(this, &SInsightsStatusBarWidget::IsPresetSet, DefaultPreset)),
-			NAME_None,
-			EUserInterfaceActionType::RadioButton
-		);
+	CreateChannelsInfo();
 
-		MenuBuilder.AddMenuEntry(
-			LOCTEXT("MemoryLabel", "Memory"),
-			LOCTEXT("MemoryLabelDesc", "Activate the memtag,memalloc,callstack,module and the default channels."),
-			FSlateIcon(),
-			FUIAction(FExecuteAction::CreateSP(this, &SInsightsStatusBarWidget::SetTraceChannels, MemoryPreset),
-				FCanExecuteAction::CreateLambda([]() {return false; }),
-				FIsActionChecked::CreateSP(this, &SInsightsStatusBarWidget::IsPresetSet, MemoryPreset)),
-			NAME_None,
-			EUserInterfaceActionType::RadioButton
-		);
-
-		MenuBuilder.AddMenuEntry(
-			LOCTEXT("TaskGraphLabel", "Task Graph"),
-			LOCTEXT("TaskGraphLabelDesc", "Activate the task and the default channels."),
-			FSlateIcon(),
-			FUIAction(FExecuteAction::CreateSP(this, &SInsightsStatusBarWidget::SetTraceChannels, TaskGraphPreset),
-				FCanExecuteAction::CreateLambda([]() {return false; }),
-				FIsActionChecked::CreateSP(this, &SInsightsStatusBarWidget::IsPresetSet, TaskGraphPreset)),
-			NAME_None,
-			EUserInterfaceActionType::RadioButton
-		);
-
-		MenuBuilder.AddMenuEntry(
-			LOCTEXT("Context Switches Label", "Context Switches"),
-			LOCTEXT("Context Switches LabelDesc", "Activate the contextswitches and the default channels."),
-			FSlateIcon(),
-			FUIAction(FExecuteAction::CreateSP(this, &SInsightsStatusBarWidget::SetTraceChannels, ContextSwitchesPreset),
-				FCanExecuteAction::CreateLambda([]() { return false; }),
-				FIsActionChecked::CreateSP(this, &SInsightsStatusBarWidget::IsPresetSet, ContextSwitchesPreset)),
-			NAME_None,
-			EUserInterfaceActionType::RadioButton
-		);
+	MenuBuilder.BeginSection("Channels", LOCTEXT("Channels", "Channels"));
+	{
+		for (int32 Index = 0; Index < ChannelsInfo.Num(); ++Index)
+		{
+			const FChannelData& Data = ChannelsInfo[Index];
+			FString ChannelDisplayName = Data.Name;
+			ChannelDisplayName.RemoveFromEnd(TEXT("Channel"), 7);
+
+			MenuBuilder.AddMenuEntry(
+				FText::FromString(ChannelDisplayName),
+				FText::Format(LOCTEXT("ChannelDesc", "Enable/disable the {0} channel"), FText::FromString(ChannelDisplayName)),
+				FSlateIcon(),
+				FUIAction(FExecuteAction::CreateSP(this, &SInsightsStatusBarWidget::ToggleChannel_Execute, Index),
+					FCanExecuteAction::CreateLambda([Value = !Data.bIsReadOnly]() { return Value; }),
+					FIsActionChecked::CreateSP(this, &SInsightsStatusBarWidget::ToggleChannel_IsChecked, Index)),
+				NAME_None,
+				EUserInterfaceActionType::ToggleButton);
+		}
 	}
 	MenuBuilder.EndSection();
 }
@@ -618,11 +598,11 @@ void SInsightsStatusBarWidget::StartTracing()
 {
 	if (TraceDestination == ETraceDestination::TraceStore)
 	{
-		FTraceAuxiliary::Start(FTraceAuxiliary::EConnectionType::Network, TEXT("localhost"), Channels);
+		FTraceAuxiliary::Start(FTraceAuxiliary::EConnectionType::Network, TEXT("localhost"), nullptr);
 	}
 	else if (TraceDestination == ETraceDestination::File)
 	{
-		FTraceAuxiliary::Start(FTraceAuxiliary::EConnectionType::File, nullptr, Channels);
+		FTraceAuxiliary::Start(FTraceAuxiliary::EConnectionType::File, nullptr, nullptr);
 	}
 }
 
@@ -646,16 +626,6 @@ EVisibility SInsightsStatusBarWidget::GetStopTraceIconVisibility() const
 	return EVisibility::Hidden;
 }
 
-void SInsightsStatusBarWidget::SetTraceChannels(const TCHAR* InChannels)
-{
-	Channels = InChannels;
-}
-
-bool SInsightsStatusBarWidget::IsPresetSet(const TCHAR* InChannels) const
-{
-	return Channels == InChannels;
-}
-
 bool SInsightsStatusBarWidget::GetBooleanSettingValue(const TCHAR* InSettingName)
 {
 	bool Value = false;
@@ -732,4 +702,76 @@ void SInsightsStatusBarWidget::CacheTraceStorePath()
 	}
 }
 
+void SInsightsStatusBarWidget::CreateChannelsInfo()
+{
+	ChannelsInfo.Empty();
+
+	UE::Trace::EnumerateChannels([](const UE::Trace::FChannelInfo& Info, void* User)
+	{
+		TArray<FChannelData>* Channels = (TArray<FChannelData>*) User;
+
+		FChannelData NewChannelInfo;
+		NewChannelInfo.Name = Info.Name;
+		NewChannelInfo.bIsEnabled = Info.bIsEnabled;
+		NewChannelInfo.bIsReadOnly = Info.bIsReadOnly;
+
+		Channels->Add(NewChannelInfo);
+		return true;
+	}, &ChannelsInfo);
+
+	Algo::SortBy(ChannelsInfo, [](const FChannelData& Entry) { return Entry.Name; }, [](const FString& Rhs, const FString& Lhs)
+		{
+			return Rhs.Compare(Lhs) < 0;
+		});
+}
+
+void SInsightsStatusBarWidget::UpdateChannelsInfo()
+{
+	UE::Trace::EnumerateChannels([](const UE::Trace::FChannelInfo& Info, void* User)
+	{
+		TArray<FChannelData>* Channels = (TArray<FChannelData>*) User;
+
+		FString Name = Info.Name;
+		int32 Index = Algo::BinarySearchBy(*Channels, Name, [](const FChannelData& Entry) { return Entry.Name; }, [](const FString& Rhs, const FString& Lhs)
+			{
+				return Rhs.Compare(Lhs) < 0;
+			});
+
+		if (Index != INDEX_NONE)
+		{
+			FChannelData& Data = (*Channels)[Index];
+			Data.Name = Info.Name;
+			Data.bIsEnabled = Info.bIsEnabled;
+			Data.bIsReadOnly = Info.bIsReadOnly;
+		}
+
+		return true;
+	}, &ChannelsInfo);
+
+	bShouldUpdateChannels = false;
+}
+
+void SInsightsStatusBarWidget::ToggleChannel_Execute(int32 Index)
+{
+	if (Index < ChannelsInfo.Num())
+	{
+		UE::Trace::ToggleChannel(*ChannelsInfo[Index].Name, !ChannelsInfo[Index].bIsEnabled);
+	}
+}
+
+bool SInsightsStatusBarWidget::ToggleChannel_IsChecked(int32 Index)
+{
+	if (bShouldUpdateChannels)
+	{
+		UpdateChannelsInfo();
+	}
+
+	if (Index < ChannelsInfo.Num())
+	{
+		return ChannelsInfo[Index].bIsEnabled;
+	}
+
+	return false;
+}
+
 #undef LOCTEXT_NAMESPACE
diff --git a/Engine/Plugins/TraceUtilities/Source/EditorTraceUtilities/Private/InsightsStatusBar/SInsightsStatusBar.h b/Engine/Plugins/TraceUtilities/Source/EditorTraceUtilities/Private/InsightsStatusBar/SInsightsStatusBar.h
--- a/Engine/Plugins/TraceUtilities/Source/EditorTraceUtilities/Private/InsightsStatusBar/SInsightsStatusBar.h
+++ b/Engine/Plugins/TraceUtilities/Source/EditorTraceUtilities/Private/InsightsStatusBar/SInsightsStatusBar.h
@@ -18,17 +18,27 @@ enum class ETraceDestination : uint32
 	File = 1
 };
 
+struct FChannelData
+{
+	FString Name;
+	bool bIsEnabled = false;
+	bool bIsReadOnly = false;
+};
+
 /**
  *  Status bar widget for Unreal Insights.
  *  Shows buttons to start tracing either to a file or to the trace store and allows saving a snapshot to file.
  */
 class SInsightsStatusBarWidget : public SCompoundWidget
 {
+public:
 	SLATE_BEGIN_ARGS(SInsightsStatusBarWidget) {}
 	SLATE_END_ARGS()
 
 	void Construct(const FArguments& InArgs);
 
+	virtual void Tick(const FGeometry& AllottedGeometry, const double InCurrentTime, const float InDeltaTime) override;
+
 private:
 	FText	GetTitleToolTipText() const;
 	
@@ -83,6 +93,12 @@ private:
 
 	void CacheTraceStorePath();
 
+	void ToggleChannel_Execute(int32 Index);
+	bool ToggleChannel_IsChecked(int32 Index);
+
+	void CreateChannelsInfo();
+	void UpdateChannelsInfo();
+
 private:
 	static const TCHAR* DefaultPreset;
 	static const TCHAR* MemoryPreset;
@@ -96,8 +112,10 @@ private:
 
 	ETraceDestination TraceDestination = ETraceDestination::TraceStore;
 	bool bIsTraceRecordButtonHovered = false;
-	const TCHAR* Channels;
 	mutable double ConnectionStartTime = 0.0f;
 
 	FString TraceStorePath;
+
+	TArray<FChannelData> ChannelsInfo;
+	bool bShouldUpdateChannels = false;
 };
diff --git a/Engine/Plugins/TraceUtilities/Source/EditorTraceUtilities/Private/UnrealInsightsLauncher.cpp b/Engine/Plugins/TraceUtilities/Source/EditorTraceUtilities/Private/UnrealInsightsLauncher.cpp
--- a/Engine/Plugins/TraceUtilities/Source/EditorTraceUtilities/Private/UnrealInsightsLauncher.cpp
+++ b/Engine/Plugins/TraceUtilities/Source/EditorTraceUtilities/Private/UnrealInsightsLauncher.cpp
@@ -242,7 +242,7 @@ bool FUnrealInsightsLauncher::OpenActiveTraceFromStore(const FString& TraceHostA
 	
 	uint32 TraceId = SessionInfo->GetTraceId();
 	
-	if (!OpenRemoteTrace(TraceHostAddress, StoreClient->GetStorePort(), TraceId))
+	if (!OpenRemoteTrace(TraceHostAddress, (uint16)StoreClient->GetStorePort(), TraceId))
 	{
 		//Failed to open Trace File %s
 		return false;
diff --git a/Engine/Source/Runtime/TraceLog/Private/Trace/Channel.cpp b/Engine/Source/Runtime/TraceLog/Private/Trace/Channel.cpp
--- a/Engine/Source/Runtime/TraceLog/Private/Trace/Channel.cpp
+++ b/Engine/Source/Runtime/TraceLog/Private/Trace/Channel.cpp
@@ -247,6 +247,34 @@ void FChannel::EnumerateChannels(ChannelIterFunc Func, void* User)
 	}
 }
 
+///////////////////////////////////////////////////////////////////////////////
+void FChannel::EnumerateChannels(ChannelIterCallback Func, void* User)
+{
+	using namespace Private;
+	FChannel* ChannelLists[] =
+	{
+		AtomicLoadAcquire(&GNewChannelList),
+		AtomicLoadAcquire(&GHeadChannel),
+	};
+
+	FChannelInfo Info;
+	for (FChannel* Channel : ChannelLists)
+	{
+		for (; Channel != nullptr; Channel = Channel->Next)
+		{
+			Info.Name = Channel->Name.Ptr;
+			Info.Desc = Channel->Args.Desc;
+			Info.bIsEnabled = Channel->IsEnabled();
+			Info.bIsReadOnly = Channel->Args.bReadOnly;
+			bool Result = Func(Info, User);
+			if (!Result)
+			{
+				return;
+			}
+		}
+	}
+}
+
 ///////////////////////////////////////////////////////////////////////////////
 bool FChannel::Toggle(bool bEnabled)
 {
diff --git a/Engine/Source/Runtime/TraceLog/Private/Trace/Trace.cpp b/Engine/Source/Runtime/TraceLog/Private/Trace/Trace.cpp
--- a/Engine/Source/Runtime/TraceLog/Private/Trace/Trace.cpp
+++ b/Engine/Source/Runtime/TraceLog/Private/Trace/Trace.cpp
@@ -135,6 +135,27 @@ bool ToggleChannel(const TCHAR* ChannelName, bool bEnabled)
 
 ////////////////////////////////////////////////////////////////////////////////
 void EnumerateChannels(ChannelIterFunc IterFunc, void* User)
+{
+	struct FCallbackDataWrapper
+	{
+		ChannelIterFunc* Func;
+		void* User;
+	};
+
+	FCallbackDataWrapper Wrapper;
+	Wrapper.Func = IterFunc;
+	Wrapper.User = User;
+
+	FChannel::EnumerateChannels([](const FChannelInfo& Info, void* User)
+		{
+			FCallbackDataWrapper* Wrapper = (FCallbackDataWrapper*)User;
+			(*Wrapper).Func(Info.Name, Info.bIsEnabled, (*Wrapper).User);
+			return true;
+		}, &Wrapper);
+}
+
+////////////////////////////////////////////////////////////////////////////////
+void EnumerateChannels(ChannelIterCallback IterFunc, void* User)
 {
 	FChannel::EnumerateChannels(IterFunc, User);
 }
diff --git a/Engine/Source/Runtime/TraceLog/Public/Trace/Detail/Channel.h b/Engine/Source/Runtime/TraceLog/Public/Trace/Detail/Channel.h
--- a/Engine/Source/Runtime/TraceLog/Public/Trace/Detail/Channel.h
+++ b/Engine/Source/Runtime/TraceLog/Public/Trace/Detail/Channel.h
@@ -10,8 +10,17 @@
 
 namespace UE {
 namespace Trace {
-	
+
+struct FChannelInfo
+{
+	const ANSICHAR* Name;
+	const ANSICHAR* Desc;
+	bool bIsEnabled;
+	bool bIsReadOnly;
+};
+
 typedef void ChannelIterFunc(const ANSICHAR*, bool, void*);	
+typedef bool ChannelIterCallback(const FChannelInfo& OutChannelInfo, void*);
 
 /*
 	A named channel which can be used to filter trace events. Channels can be
@@ -52,7 +61,9 @@ public:
 	static bool			Toggle(const ANSICHAR* ChannelName, bool bEnabled);
 	static void			ToggleAll(bool bEnabled);
 	static FChannel*	FindChannel(const ANSICHAR* ChannelName);
+	UE_DEPRECATED(5.2, "Please use the ChannelIterCallback overload for enumerating channels.")
 	static void			EnumerateChannels(ChannelIterFunc Func, void* User);
+	static void			EnumerateChannels(ChannelIterCallback Func, void* User);
 	bool				Toggle(bool bEnabled);
 	bool				IsEnabled() const;
 	explicit			operator bool () const;
diff --git a/Engine/Source/Runtime/TraceLog/Public/Trace/Trace.h b/Engine/Source/Runtime/TraceLog/Public/Trace/Trace.h
--- a/Engine/Source/Runtime/TraceLog/Public/Trace/Trace.h
+++ b/Engine/Source/Runtime/TraceLog/Public/Trace/Trace.h
@@ -92,9 +92,14 @@ struct FInitializeDesc
 	OnConnectFunc*	OnConnectionFunc	= nullptr;
 };
 
+struct FChannelInfo;
+
 typedef void*		AllocFunc(SIZE_T, uint32);
 typedef void		FreeFunc(void*, SIZE_T);
-typedef void		ChannelIterFunc(const ANSICHAR*, bool, void*);	
+typedef void		ChannelIterFunc(const ANSICHAR*, bool, void*);
+/*  The callback provides information about a channel and a user provided pointer. 
+	Returning false from the callback will stop the enumeration */
+typedef bool		ChannelIterCallback(const FChannelInfo& Info, void*/*User*/);
 
 struct FStatistics
 {
@@ -127,6 +132,7 @@ UE_TRACE_API bool	Stop() UE_TRACE_IMPL(false);
 UE_TRACE_API bool	IsChannel(const TCHAR* ChanneName) UE_TRACE_IMPL(false);
 UE_TRACE_API bool	ToggleChannel(const TCHAR* ChannelName, bool bEnabled) UE_TRACE_IMPL(false);
 UE_TRACE_API void	EnumerateChannels(ChannelIterFunc IterFunc, void* User) UE_TRACE_IMPL();
+UE_TRACE_API void	EnumerateChannels(ChannelIterCallback IterFunc, void* User) UE_TRACE_IMPL();
 UE_TRACE_API void	ThreadRegister(const TCHAR* Name, uint32 SystemId, int32 SortHint) UE_TRACE_IMPL();
 UE_TRACE_API void	ThreadGroupBegin(const TCHAR* Name) UE_TRACE_IMPL();
 UE_TRACE_API void	ThreadGroupEnd() UE_TRACE_IMPL();
