Fix Open Live Session on Trace Start opening the session in a invalid state.
Fix CIS error.

diff --git a/Engine/Plugins/TraceUtilities/Source/EditorTraceUtilities/Private/InsightsStatusBar/SInsightsStatusBar.cpp b/Engine/Plugins/TraceUtilities/Source/EditorTraceUtilities/Private/InsightsStatusBar/SInsightsStatusBar.cpp
--- a/Engine/Plugins/TraceUtilities/Source/EditorTraceUtilities/Private/InsightsStatusBar/SInsightsStatusBar.cpp
+++ b/Engine/Plugins/TraceUtilities/Source/EditorTraceUtilities/Private/InsightsStatusBar/SInsightsStatusBar.cpp
@@ -10,6 +10,8 @@
 #include "Framework/Notifications/NotificationManager.h"
 #include "GenericPlatform/GenericPlatformFile.h"
 #include "HAL/FileManager.h"
+#include "Trace/Detail/Channel.h"
+#include "Trace/Trace.h"
 #include "Math/Color.h"
 #include "Misc/ConfigCacheIni.h"
 #include "Misc/Paths.h"
@@ -38,6 +40,22 @@ const TCHAR* SInsightsStatusBarWidget::OpenLiveSessionOnTraceStartSettingName =
 const TCHAR* SInsightsStatusBarWidget::OpenInsightsAfterTraceSettingName = TEXT("OpenInsightsAfterTrace");
 const TCHAR* SInsightsStatusBarWidget::ShowInExplorerAfterTraceSettingName = TEXT("ShowInExplorerAfterTrace");
 
+class FOpenLiveSessionTask
+{
+public:
+	FOpenLiveSessionTask()
+	{}
+
+	FORCEINLINE TStatId GetStatId() const { RETURN_QUICK_DECLARE_CYCLE_STAT(FOpenLiveSessionTask, STATGROUP_TaskGraphTasks); }
+	ENamedThreads::Type GetDesiredThread() { return ENamedThreads::Type::GameThread; }
+	static ESubsequentsMode::Type GetSubsequentsMode() { return ESubsequentsMode::TrackSubsequents; }
+
+	void DoTask(ENamedThreads::Type CurrentThread, const FGraphEventRef& MyCompletionGraphEvent)
+	{
+		FUnrealInsightsLauncher::Get()->TryOpenTraceFromDestination(FTraceAuxiliary::GetTraceDestination());
+	}
+};
+
 TSharedRef<SWidget> CreateInsightsStatusBarWidget()
 {
 	return SNew(SInsightsStatusBarWidget);
@@ -644,7 +662,7 @@ void SInsightsStatusBarWidget::OnTraceStarted(FTraceAuxiliary::EConnectionType I
 {
 	if (InTraceType == FTraceAuxiliary::EConnectionType::Network && GetBooleanSettingValue(OpenLiveSessionOnTraceStartSettingName))
 	{
-		OpenLiveSession();
+		TGraphTask<FOpenLiveSessionTask>::CreateTask().ConstructAndDispatchWhenReady();
 	}
 }
 
@@ -704,6 +722,7 @@ void SInsightsStatusBarWidget::CacheTraceStorePath()
 
 void SInsightsStatusBarWidget::CreateChannelsInfo()
 {
+#if UE_TRACE_ENABLED
 	ChannelsInfo.Empty();
 
 	UE::Trace::EnumerateChannels([](const UE::Trace::FChannelInfo& Info, void* User)
@@ -723,10 +742,12 @@ void SInsightsStatusBarWidget::CreateChannelsInfo()
 		{
 			return Rhs.Compare(Lhs) < 0;
 		});
+#endif
 }
 
 void SInsightsStatusBarWidget::UpdateChannelsInfo()
 {
+#if UE_TRACE_ENABLED
 	UE::Trace::EnumerateChannels([](const UE::Trace::FChannelInfo& Info, void* User)
 	{
 		TArray<FChannelData>* Channels = (TArray<FChannelData>*) User;
@@ -749,6 +770,7 @@ void SInsightsStatusBarWidget::UpdateChannelsInfo()
 	}, &ChannelsInfo);
 
 	bShouldUpdateChannels = false;
+#endif
 }
 
 void SInsightsStatusBarWidget::ToggleChannel_Execute(int32 Index)
