Add Trace Bookmark and Trace Screenshot buttons to the Insigths Status Bar Widget.

diff --git a/Engine/Plugins/TraceUtilities/Source/EditorTraceUtilities/Private/InsightsStatusBar/EditorTraceUtilitiesStyle.h b/Engine/Plugins/TraceUtilities/Source/EditorTraceUtilities/Private/InsightsStatusBar/EditorTraceUtilitiesStyle.h
--- a/Engine/Plugins/TraceUtilities/Source/EditorTraceUtilities/Private/InsightsStatusBar/EditorTraceUtilitiesStyle.h
+++ b/Engine/Plugins/TraceUtilities/Source/EditorTraceUtilities/Private/InsightsStatusBar/EditorTraceUtilitiesStyle.h
@@ -46,6 +46,9 @@ public:
 		Set("Icons.TraceStore.Menu", new CORE_IMAGE_BRUSH_SVG("Starship/Insights/TraceStore", Icon16x16));
 		Set("Icons.File.Menu", new CORE_IMAGE_BRUSH_SVG("Starship/Common/file", Icon16x16));
 
+		Set("Icons.Screenshot.Menu", new CORE_IMAGE_BRUSH_SVG("../Editor/Slate/Starship/Common/HighResolutionScreenshot", Icon16x16));
+		Set("Icons.Bookmark.Menu", new CORE_IMAGE_BRUSH_SVG("../Editor/Slate/Starship/Common/Bookmarks", Icon16x16));
+
 		FSlateStyleRegistry::RegisterSlateStyle(*this);
 	}
 
diff --git a/Engine/Plugins/TraceUtilities/Source/EditorTraceUtilities/Private/InsightsStatusBar/SInsightsStatusBar.cpp b/Engine/Plugins/TraceUtilities/Source/EditorTraceUtilities/Private/InsightsStatusBar/SInsightsStatusBar.cpp
--- a/Engine/Plugins/TraceUtilities/Source/EditorTraceUtilities/Private/InsightsStatusBar/SInsightsStatusBar.cpp
+++ b/Engine/Plugins/TraceUtilities/Source/EditorTraceUtilities/Private/InsightsStatusBar/SInsightsStatusBar.cpp
@@ -6,6 +6,8 @@
 #include "EditorTraceUtilitiesStyle.h"
 #include "EditorTraceUtilities.h"
 #include "Framework/Application/SlateApplication.h"
+#include "Framework/Commands/Commands.h"
+#include "Framework/Commands/UICommandList.h"
 #include "Framework/MultiBox/MultiBoxBuilder.h"
 #include "Framework/Notifications/NotificationManager.h"
 #include "GenericPlatform/GenericPlatformFile.h"
@@ -14,8 +16,11 @@
 #include "Trace/Trace.h"
 #include "Math/Color.h"
 #include "Misc/ConfigCacheIni.h"
+#include "Misc/DateTime.h"
 #include "Misc/Paths.h"
 #include "Modules/ModuleManager.h"
+#include "ProfilingDebugging/MiscTrace.h"
+#include "ProfilingDebugging/TraceScreenshot.h"
 #include "SRecentTracesList.h"
 #include "Styling/StyleColors.h"
 #include "ToolMenus.h"
@@ -56,6 +61,40 @@ public:
 	}
 };
 
+
+////////////////////////////////////////////////////////////////////////////////////////////////////
+// FInsightsStatusBarWidgetCommands
+////////////////////////////////////////////////////////////////////////////////////////////////////
+
+class FInsightsStatusBarWidgetCommands : public TCommands<FInsightsStatusBarWidgetCommands>
+{
+public:
+	FInsightsStatusBarWidgetCommands()
+		: TCommands<FInsightsStatusBarWidgetCommands>(
+			TEXT("InsightsStatusBarWidgetCommands"),
+			NSLOCTEXT("Contexts", "InsightsStatusBarWidgetCommands", "Insights Status Bar"),
+			NAME_None,
+			FEditorTraceUtilitiesStyle::Get().GetStyleSetName())
+	{
+	}
+
+	virtual ~FInsightsStatusBarWidgetCommands()
+	{
+	}
+
+	// UI_COMMAND takes long for the compiler to optimize
+	UE_DISABLE_OPTIMIZATION_SHIP
+	virtual void RegisterCommands() override
+	{
+		UI_COMMAND(Command_TraceScreenshot, "Trace Screenshot", "Takes a screenshot and sends it to the trace.", EUserInterfaceActionType::Button, FInputChord(EModifierKey::Control, EKeys::F9));
+		UI_COMMAND(Command_TraceBookmark, "Trace Bookmark", "Traces a bookmark.", EUserInterfaceActionType::Button, FInputChord());
+	}
+	UE_ENABLE_OPTIMIZATION_SHIP
+
+	TSharedPtr<FUICommandInfo> Command_TraceScreenshot;
+	TSharedPtr<FUICommandInfo> Command_TraceBookmark;
+};
+
 TSharedRef<SWidget> CreateInsightsStatusBarWidget()
 {
 	return SNew(SInsightsStatusBarWidget);
@@ -190,14 +229,43 @@ void SInsightsStatusBarWidget::Construct(const FArguments& InArgs)
 	FTraceAuxiliary::OnTraceStarted.AddSP(this, &SInsightsStatusBarWidget::OnTraceStarted);
 	FTraceAuxiliary::OnTraceStopped.AddSP(this, &SInsightsStatusBarWidget::OnTraceStopped);
 	FTraceAuxiliary::OnSnapshotSaved.AddSP(this, &SInsightsStatusBarWidget::OnSnapshotSaved);
+
+	InitCommandList();
 }
 
 TSharedRef<SWidget> SInsightsStatusBarWidget::MakeTraceMenu()
 {
-	FMenuBuilder MenuBuilder(true, nullptr);
+	FMenuBuilder MenuBuilder(true, CommandList.ToSharedRef());
 
-	MenuBuilder.BeginSection("TraceDataFiltering", LOCTEXT("TraceMenu_Section_DataFiltering", "Trace Data Filtering"));
+	MenuBuilder.BeginSection("TraceData", LOCTEXT("TraceMenu_Section_Data", "Trace Data"));
 	{
+		MenuBuilder.AddSubMenu
+		(
+			LOCTEXT("Channels", "Channels"),
+			LOCTEXT("Channels_Desc", "Select what trace channels to enable when tracing."),
+			FNewMenuDelegate::CreateSP(this, &SInsightsStatusBarWidget::Channels_BuildMenu),
+			false,
+			FSlateIcon(FEditorTraceUtilitiesStyle::Get().GetStyleSetName(), ("Icons.Trace.Menu"))
+		);
+
+		MenuBuilder.AddMenuEntry
+		(
+			FInsightsStatusBarWidgetCommands::Get().Command_TraceScreenshot,
+			NAME_None,
+			TAttribute<FText>(),
+			TAttribute<FText>(),
+			FSlateIcon(FEditorTraceUtilitiesStyle::Get().GetStyleSetName(), "Icons.Screenshot.Menu")
+		);
+
+		MenuBuilder.AddMenuEntry
+		(
+			FInsightsStatusBarWidgetCommands::Get().Command_TraceBookmark,
+			NAME_None,
+			TAttribute<FText>(),
+			TAttribute<FText>(),
+			FSlateIcon(FEditorTraceUtilitiesStyle::Get().GetStyleSetName(), "Icons.Bookmark.Menu")
+		);
+
 		MenuBuilder.AddMenuEntry(
 			LOCTEXT("StatNamedEventsLabel", "Stat Named Events"),
 			LOCTEXT("StatNamedEventsDesc", "Enable or disable named events in the stats system."),
@@ -208,15 +276,6 @@ TSharedRef<SWidget> SInsightsStatusBarWidget::MakeTraceMenu()
 			NAME_None,
 			EUserInterfaceActionType::ToggleButton
 		);
-
-		MenuBuilder.AddSubMenu
-		(
-			LOCTEXT("Channels", "Channels"),
-			LOCTEXT("Channels_Desc", "Select what trace channels to enable when tracing."),
-			FNewMenuDelegate::CreateSP(this, &SInsightsStatusBarWidget::Channels_BuildMenu),
-			false,
-			FSlateIcon(FEditorTraceUtilitiesStyle::Get().GetStyleSetName(), ("Icons.Trace.Menu"))
-		);
 	}
 	MenuBuilder.EndSection();
 
@@ -404,6 +463,14 @@ void SInsightsStatusBarWidget::Traces_BuildMenu(FMenuBuilder& MenuBuilder)
 		FText(), true);
 }
 
+void SInsightsStatusBarWidget::InitCommandList()
+{
+	FInsightsStatusBarWidgetCommands::Register();
+	CommandList = MakeShared<FUICommandList>();
+	CommandList->MapAction(FInsightsStatusBarWidgetCommands::Get().Command_TraceScreenshot, FExecuteAction::CreateSP(this, &SInsightsStatusBarWidget::TraceScreenshot_Execute), FCanExecuteAction::CreateSP(this, &SInsightsStatusBarWidget::TraceScreenshot_CanExecute));
+	CommandList->MapAction(FInsightsStatusBarWidgetCommands::Get().Command_TraceBookmark, FExecuteAction::CreateSP(this, &SInsightsStatusBarWidget::TraceBookmark_Execute), FCanExecuteAction::CreateSP(this, &SInsightsStatusBarWidget::TraceBookmark_CanExecute));
+}
+
 FText SInsightsStatusBarWidget::GetTitleToolTipText() const
 {
 	FTextBuilder DescBuilder;
@@ -796,4 +863,26 @@ bool SInsightsStatusBarWidget::ToggleChannel_IsChecked(int32 Index)
 	return false;
 }
 
+
+bool SInsightsStatusBarWidget::TraceScreenshot_CanExecute()
+{
+	return SHOULD_TRACE_SCREENSHOT();
+}
+
+void SInsightsStatusBarWidget::TraceScreenshot_Execute()
+{
+	FTraceScreenshot::RequestScreenshot(TEXT(""), false);
+}
+
+bool SInsightsStatusBarWidget::TraceBookmark_CanExecute()
+{
+	return SHOULD_TRACE_BOOKMARK();
+}
+
+void SInsightsStatusBarWidget::TraceBookmark_Execute()
+{
+	FString Bookmark = FDateTime::Now().ToString(TEXT("Bookmark_%Y%m%d_%H%M%S"));
+	TRACE_BOOKMARK(*Bookmark);
+}
+
 #undef LOCTEXT_NAMESPACE
diff --git a/Engine/Plugins/TraceUtilities/Source/EditorTraceUtilities/Private/InsightsStatusBar/SInsightsStatusBar.h b/Engine/Plugins/TraceUtilities/Source/EditorTraceUtilities/Private/InsightsStatusBar/SInsightsStatusBar.h
--- a/Engine/Plugins/TraceUtilities/Source/EditorTraceUtilities/Private/InsightsStatusBar/SInsightsStatusBar.h
+++ b/Engine/Plugins/TraceUtilities/Source/EditorTraceUtilities/Private/InsightsStatusBar/SInsightsStatusBar.h
@@ -9,6 +9,7 @@
 #include "Widgets/SCompoundWidget.h"
 
 class FMenuBuilder;
+class FUICommandList;
 
 TSharedRef<SWidget> CreateInsightsStatusBarWidget();
 
@@ -99,6 +100,14 @@ private:
 	void CreateChannelsInfo();
 	void UpdateChannelsInfo();
 
+	void InitCommandList();
+
+	bool TraceScreenshot_CanExecute();
+	void TraceScreenshot_Execute();
+
+	bool TraceBookmark_CanExecute();
+	void TraceBookmark_Execute();
+
 private:
 	static const TCHAR* DefaultPreset;
 	static const TCHAR* MemoryPreset;
@@ -118,4 +127,6 @@ private:
 
 	TArray<FChannelData> ChannelsInfo;
 	bool bShouldUpdateChannels = false;
+
+	TSharedPtr<FUICommandList> CommandList;
 };
diff --git a/Engine/Source/Runtime/Core/Private/ProfilingDebugging/MiscTrace.cpp b/Engine/Source/Runtime/Core/Private/ProfilingDebugging/MiscTrace.cpp
--- a/Engine/Source/Runtime/Core/Private/ProfilingDebugging/MiscTrace.cpp
+++ b/Engine/Source/Runtime/Core/Private/ProfilingDebugging/MiscTrace.cpp
@@ -139,4 +139,9 @@ bool FMiscTrace::ShouldTraceScreenshot()
 {
 	return UE_TRACE_CHANNELEXPR_IS_ENABLED(ScreenshotChannel);
 }
+
+bool FMiscTrace::ShouldTraceBookmark()
+{
+	return UE_TRACE_CHANNELEXPR_IS_ENABLED(BookmarkChannel);
+}
 #endif
diff --git a/Engine/Source/Runtime/Core/Public/ProfilingDebugging/MiscTrace.h b/Engine/Source/Runtime/Core/Public/ProfilingDebugging/MiscTrace.h
--- a/Engine/Source/Runtime/Core/Public/ProfilingDebugging/MiscTrace.h
+++ b/Engine/Source/Runtime/Core/Public/ProfilingDebugging/MiscTrace.h
@@ -66,6 +66,7 @@ struct FMiscTrace
 
 	CORE_API static void OutputScreenshot(const TCHAR* Name, uint64 Cycle, uint32 Width, uint32 Height, TArray64<uint8> Data);
 	CORE_API static bool ShouldTraceScreenshot();
+	CORE_API static bool ShouldTraceBookmark();
 
 private:
 	CORE_API static void OutputBookmarkInternal(const void* BookmarkPoint, uint16 EncodedFormatArgsSize, uint8* EncodedFormatArgs);
@@ -92,6 +93,9 @@ private:
 #define SHOULD_TRACE_SCREENSHOT() \
 	FMiscTrace::ShouldTraceScreenshot()
 
+#define SHOULD_TRACE_BOOKMARK() \
+	FMiscTrace::ShouldTraceBookmark()
+
 #else
 
 #define TRACE_BOOKMARK(...)
@@ -99,5 +103,6 @@ private:
 #define TRACE_END_FRAME(...)
 #define TRACE_SCREENSHOT(...)
 #define SHOULD_TRACE_SCREENSHOT(...) false
+#define SHOULD_TRACE_BOOKMARK(...) false
 
 #endif
