Redesign the recent traces menu.

diff --git a/Engine/Plugins/TraceUtilities/Source/EditorTraceUtilities/Private/InsightsStatusBar/SInsightsStatusBar.cpp b/Engine/Plugins/TraceUtilities/Source/EditorTraceUtilities/Private/InsightsStatusBar/SInsightsStatusBar.cpp
--- a/Engine/Plugins/TraceUtilities/Source/EditorTraceUtilities/Private/InsightsStatusBar/SInsightsStatusBar.cpp
+++ b/Engine/Plugins/TraceUtilities/Source/EditorTraceUtilities/Private/InsightsStatusBar/SInsightsStatusBar.cpp
@@ -451,16 +451,18 @@ void SInsightsStatusBarWidget::Channels_BuildMenu(FMenuBuilder& MenuBuilder)
 void SInsightsStatusBarWidget::Traces_BuildMenu(FMenuBuilder& MenuBuilder)
 {
 	CacheTraceStorePath();
+	PopulateRecentTracesList();
 
-	MenuBuilder.AddWidget(
-		SNew(SBox)
-		.MaxDesiredHeight(400.0f)
-		.MinDesiredWidth(200.0f)
-		.MaxDesiredWidth(270.0f)
-		[
-			SNew(SRecentTracesList, TraceStorePath)
-		],
-		FText(), true);
+	for (int32 Index = 0; Index < Traces.Num(); ++Index)
+	{
+		MenuBuilder.AddMenuEntry(
+			FUIAction(FExecuteAction::CreateSP(this, &SInsightsStatusBarWidget::OpenTrace, Index),
+					  FCanExecuteAction()),
+			SNew(SRecentTracesListEntry, Traces[Index], TraceStorePath),
+			NAME_None,
+			FText::FromString(Traces[Index]->FilePath),
+			EUserInterfaceActionType::Button);
+	}
 }
 
 void SInsightsStatusBarWidget::InitCommandList()
@@ -885,4 +887,39 @@ void SInsightsStatusBarWidget::TraceBookmark_Execute()
 	TRACE_BOOKMARK(*Bookmark);
 }
 
+void SInsightsStatusBarWidget::PopulateRecentTracesList()
+{
+	Traces.Empty();
+	bool bIsFromTraceStore = true;
+	auto Visitor = [this, &bIsFromTraceStore](const TCHAR* Filename, const FFileStatData& StatData)
+	{
+		if (FPaths::GetExtension(Filename) == TEXT("utrace"))
+		{
+			TSharedPtr<FTraceFileInfo> TraceInfo = MakeShared<FTraceFileInfo>();
+			TraceInfo->FilePath = Filename;
+			TraceInfo->ModifiedTime = StatData.ModificationTime;
+			TraceInfo->bIsFromTraceStore = bIsFromTraceStore;
+
+			Traces.Add(TraceInfo);
+		}
+
+		return true;
+	};
+
+	IFileManager::Get().IterateDirectoryStat(*TraceStorePath, Visitor);
+
+	bIsFromTraceStore = false;
+	IFileManager::Get().IterateDirectoryStat(*FPaths::ProfilingDir(), Visitor);
+
+	Algo::SortBy(Traces, [](TSharedPtr<FTraceFileInfo> TraceInfo) { return *TraceInfo; });
+}
+
+void SInsightsStatusBarWidget::OpenTrace(int32 Index)
+{
+	if (Index < Traces.Num())
+	{
+		FUnrealInsightsLauncher::Get()->TryOpenTraceFromDestination(Traces[Index]->FilePath);
+	}
+}
+
 #undef LOCTEXT_NAMESPACE
diff --git a/Engine/Plugins/TraceUtilities/Source/EditorTraceUtilities/Private/InsightsStatusBar/SInsightsStatusBar.h b/Engine/Plugins/TraceUtilities/Source/EditorTraceUtilities/Private/InsightsStatusBar/SInsightsStatusBar.h
--- a/Engine/Plugins/TraceUtilities/Source/EditorTraceUtilities/Private/InsightsStatusBar/SInsightsStatusBar.h
+++ b/Engine/Plugins/TraceUtilities/Source/EditorTraceUtilities/Private/InsightsStatusBar/SInsightsStatusBar.h
@@ -26,6 +26,18 @@ struct FChannelData
 	bool bIsReadOnly = false;
 };
 
+struct FTraceFileInfo
+{
+	FString FilePath;
+	FDateTime ModifiedTime;
+	bool bIsFromTraceStore;
+
+	bool operator <(const FTraceFileInfo& rhs)
+	{
+		return this->ModifiedTime > rhs.ModifiedTime;
+	}
+};
+
 /**
  *  Status bar widget for Unreal Insights.
  *  Shows buttons to start tracing either to a file or to the trace store and allows saving a snapshot to file.
@@ -108,6 +120,10 @@ private:
 	bool TraceBookmark_CanExecute();
 	void TraceBookmark_Execute();
 
+	void PopulateRecentTracesList();
+
+	void OpenTrace(int32 Index);
+
 private:
 	static const TCHAR* DefaultPreset;
 	static const TCHAR* MemoryPreset;
@@ -129,4 +145,6 @@ private:
 	bool bShouldUpdateChannels = false;
 
 	TSharedPtr<FUICommandList> CommandList;
+
+	TArray<TSharedPtr<FTraceFileInfo>> Traces;
 };
diff --git a/Engine/Plugins/TraceUtilities/Source/EditorTraceUtilities/Private/InsightsStatusBar/SRecentTracesList.cpp b/Engine/Plugins/TraceUtilities/Source/EditorTraceUtilities/Private/InsightsStatusBar/SRecentTracesList.cpp
--- a/Engine/Plugins/TraceUtilities/Source/EditorTraceUtilities/Private/InsightsStatusBar/SRecentTracesList.cpp
+++ b/Engine/Plugins/TraceUtilities/Source/EditorTraceUtilities/Private/InsightsStatusBar/SRecentTracesList.cpp
@@ -8,206 +8,99 @@
 #include "GenericPlatform/GenericPlatformFile.h"
 #include "HAL/FileManager.h"
 #include "Misc/Paths.h"
+#include "SInsightsStatusBar.h"
 #include "UnrealInsightsLauncher.h"
 #include "Widgets/Input/SButton.h"
 #include "Widgets/Layout/SScrollBorder.h"
 
 #define LOCTEXT_NAMESPACE "RecentTracesList"
 
-class SRecentTracesListEntry : public STableRow<TSharedPtr<FString>>
+void SRecentTracesListEntry::Construct(const FArguments& InArgs, TSharedPtr<FTraceFileInfo> InTrace, const FString& InStorePath)
 {
-public:
-	SLATE_BEGIN_ARGS(SRecentTracesListEntry) {}
-
-	SLATE_ARGUMENT(TSharedPtr<FTraceFileInfo>, TraceInfo)
-
-	SLATE_END_ARGS()
+	TraceInfo = InTrace;
+	StorePath = InStorePath;
 
-	void Construct(const FArguments& InArgs, const TSharedRef<STableViewBase>& InOwnerTable)
+	FString TraceName = FPaths::GetBaseFilename(TraceInfo->FilePath);
+	TSharedRef<FSlateFontMeasure> FontMeasureService = FSlateApplication::Get().GetRenderer()->GetFontMeasureService();
+	int32 LastCharacter = FontMeasureService->FindLastWholeCharacterIndexBeforeOffset(TraceName, FAppStyle::GetFontStyle("Menu.Label"), 180);
+	if (LastCharacter < TraceName.Len() - 1 && LastCharacter > 3)
 	{
-		TraceInfo = InArgs._TraceInfo;
-
-		STableRow<TSharedPtr<FString>>::Construct(STableRow<TSharedPtr<FString>>::FArguments(), InOwnerTable);
+		TraceName.LeftInline(LastCharacter - 3);
+		TraceName += TEXT("...");
 	}
 
-	virtual void ConstructChildren(ETableViewMode::Type InOwnerTableMode, const TAttribute<FMargin>& InPadding, const TSharedRef<SWidget>& InContent) override
-	{
-		FString TraceName = FPaths::GetBaseFilename(TraceInfo->FilePath);
-		TSharedRef<FSlateFontMeasure> FontMeasureService = FSlateApplication::Get().GetRenderer()->GetFontMeasureService();
-		int32 LastCharacter = FontMeasureService->FindLastWholeCharacterIndexBeforeOffset(TraceName, FAppStyle::GetFontStyle("Menu.Label"), 180);
-		if (LastCharacter < TraceName.Len() - 1 && LastCharacter > 3)
-		{
-			TraceName.LeftInline(LastCharacter - 3);
-			TraceName += TEXT("...");
-		}
-
-		FString TooltipText = TraceInfo->FilePath;
-		FPaths::NormalizeFilename(TooltipText);
-		TooltipText = FPaths::ConvertRelativePathToFull(TooltipText);
-
-		const FSlateBrush* TraceLocationIcon;
-		FText TraceLocationTooltip;
-
-		if (TraceInfo->bIsFromTraceStore)
-		{
-			TraceLocationIcon = FEditorTraceUtilitiesStyle::Get().GetBrush("Icons.TraceStore.Menu");
-			TraceLocationTooltip = LOCTEXT("TraceStoreLocationTooltip", "This trace is located in the trace store.");
-		}
-		else
-		{
-			TraceLocationIcon = FEditorTraceUtilitiesStyle::Get().GetBrush("Icons.File.Menu");
-			TraceLocationTooltip = LOCTEXT("TraceFileLocationTooltip", "This trace was saved to file and is located in the current project's profilling folder.");
-		}
+	FString TooltipText = TraceInfo->FilePath;
+	FPaths::NormalizeFilename(TooltipText);
+	TooltipText = FPaths::ConvertRelativePathToFull(TooltipText);
 
-		ChildSlot
-			[
-				SNew(SBorder)
-				.BorderImage(FAppStyle::GetBrush("Menu.Background"))
-				[
-					SNew(SHorizontalBox)
+	const FSlateBrush* TraceLocationIcon;
+	FText TraceLocationTooltip;
 
-					+ SHorizontalBox::Slot()
-					.VAlign(EVerticalAlignment::VAlign_Center)
-					.HAlign(EHorizontalAlignment::HAlign_Left)
-					.Padding(8.0f, 0.0f, 0.0f, 0.0f)
-					.AutoWidth()
-					[
-						SNew(SImage)
-						.ColorAndOpacity(FSlateColor::UseSubduedForeground())
-						.Image(TraceLocationIcon)
-						.ToolTipText(TraceLocationTooltip)
-					]
-
-					+ SHorizontalBox::Slot()
-					.VAlign(EVerticalAlignment::VAlign_Center)
-					.HAlign(EHorizontalAlignment::HAlign_Left)
-					.Padding(0.0f, 2.0f, 0.0f, 0.0f)
-					.AutoWidth()
-					[
-						SNew(STextBlock)
-						.TextStyle(FAppStyle::Get(), "Menu.Label")
-						.Text(FText::FromString(TraceName))
-						.ToolTipText(FText::FromString(TooltipText))
-						.ColorAndOpacity(FSlateColor::UseForeground())
-					]
-					
-					+ SHorizontalBox::Slot()
-					.FillWidth(1.0f)
-					.HAlign(EHorizontalAlignment::HAlign_Right)
-					[
-						SNew(SHorizontalBox)
-
-						+ SHorizontalBox::Slot()
-						.HAlign(EHorizontalAlignment::HAlign_Right)
-						.FillWidth(1.0f)
-						[
-							SNew(SButton)
-							.ButtonStyle(FAppStyle::Get(), "SimpleButton")
-							.OnClicked(this, &SRecentTracesListEntry::OpenContainingFolder)
-							.ToolTipText(LOCTEXT("ExploreFolderTooltip", "Open the folder containing the trace file."))
-							.Content()
-							[
-								SNew(SImage)
-								.ColorAndOpacity(FSlateColor::UseSubduedForeground())
-								.Image(FAppStyle::Get().GetBrush("Icons.FolderOpen"))
-							]
-						]
-
-						+ SHorizontalBox::Slot()
-						.HAlign(EHorizontalAlignment::HAlign_Right)
-						.AutoWidth()
-						[
-							SNew(SButton)
-							.ButtonStyle(FAppStyle::Get(), "SimpleButton")
-							.OnClicked(this, &SRecentTracesListEntry::OpenTrace)
-							.ToolTipText(LOCTEXT("OpenTraceTooltip", "Open the trace file in Unreal Insights."))
-							.Content()
-							[
-								SNew(SImage)
-								.ColorAndOpacity(FSlateColor::UseSubduedForeground())
-								.Image(FEditorTraceUtilitiesStyle::Get().GetBrush("Icons.UnrealInsights.Menu"))
-							]
-						]
-					]
-				]
-			];
-	}
-
-	FReply OpenContainingFolder()
+	if (TraceInfo->bIsFromTraceStore)
 	{
-		FPlatformProcess::ExploreFolder(*TraceInfo->FilePath);
-
-		return FReply::Handled();
+		TraceLocationIcon = FEditorTraceUtilitiesStyle::Get().GetBrush("Icons.TraceStore.Menu");
+		TraceLocationTooltip = LOCTEXT("TraceStoreLocationTooltip", "This trace is located in the trace store.");
 	}
-
-	FReply OpenTrace()
+	else
 	{
-		FUnrealInsightsLauncher::Get()->TryOpenTraceFromDestination(TraceInfo->FilePath);
-		return FReply::Handled();
+		TraceLocationIcon = FEditorTraceUtilitiesStyle::Get().GetBrush("Icons.File.Menu");
+		TraceLocationTooltip = LOCTEXT("TraceFileLocationTooltip", "This trace was saved to file and is located in the current project's profilling folder.");
 	}
 
-	TSharedPtr<FTraceFileInfo> TraceInfo;
-};
-
-void SRecentTracesList::Construct(const FArguments& InArgs, const FString& InStorePath)
-{
-	StorePath = InStorePath;
-
-	PopulateRecentTracesList();
-
-	ListView = SNew(SListView<TSharedPtr<FTraceFileInfo>>)
-		.IsFocusable(true)
-		.SelectionMode(ESelectionMode::None)
-		.ListItemsSource(&Traces)
-		.OnGenerateRow(this, &SRecentTracesList::OnGenerateRow);
-	
 	ChildSlot
-	[
-		SNew(SVerticalBox)
-
-		+ SVerticalBox::Slot()
-		.Padding(FMargin(2.0f, 0.0f))
-		.FillHeight(1.0f)
 		[
-			SNew(SScrollBorder, ListView.ToSharedRef())
+			SNew(SHorizontalBox)
+
+			+ SHorizontalBox::Slot()
+			.VAlign(EVerticalAlignment::VAlign_Center)
+			.HAlign(EHorizontalAlignment::HAlign_Left)
+			.Padding(8.0f, 0.0f, 0.0f, 0.0f)
+			.AutoWidth()
 			[
-				ListView.ToSharedRef()
+				SNew(SImage)
+				.ColorAndOpacity(FSlateColor::UseSubduedForeground())
+				.Image(TraceLocationIcon)
+				.ToolTipText(TraceLocationTooltip)
 			]
-		]
-	];
-}
 
-TSharedRef<ITableRow> SRecentTracesList::OnGenerateRow(TSharedPtr<FTraceFileInfo> Item, const TSharedRef<STableViewBase>& OwnerTable)
-{
-	return SNew(SRecentTracesListEntry, OwnerTable)
-		.TraceInfo(Item);
+			+ SHorizontalBox::Slot()
+			.VAlign(EVerticalAlignment::VAlign_Center)
+			.HAlign(EHorizontalAlignment::HAlign_Left)
+			.Padding(0.0f, 2.0f, 0.0f, 0.0f)
+			.AutoWidth()
+			[
+				SNew(STextBlock)
+				.TextStyle(FAppStyle::Get(), "Menu.Label")
+				.Text(FText::FromString(TraceName))
+				.ToolTipText(FText::FromString(TooltipText))
+				.ColorAndOpacity(FSlateColor::UseForeground())
+			]
+					
+			+ SHorizontalBox::Slot()
+			.VAlign(EVerticalAlignment::VAlign_Center)
+			.HAlign(EHorizontalAlignment::HAlign_Left)
+			.AutoWidth()
+			[
+				SNew(SButton)
+				.ButtonStyle(FAppStyle::Get(), "SimpleButton")
+				.OnClicked(this, &SRecentTracesListEntry::OpenContainingFolder)
+				.ToolTipText(LOCTEXT("ExploreFolderTooltip", "Open the folder containing the trace file."))
+				.Visibility_Lambda([this]() {return this->IsHovered() ? EVisibility::Visible : EVisibility::Collapsed; })
+				.Content()
+				[
+					SNew(SImage)
+					.ColorAndOpacity(FSlateColor::UseSubduedForeground())
+					.Image(FAppStyle::Get().GetBrush("Icons.FolderOpen"))
+				]
+			]
+		];
 }
 
-void SRecentTracesList::PopulateRecentTracesList()
+FReply SRecentTracesListEntry::OpenContainingFolder()
 {
-	Traces.Empty();
-	bool bIsFromTraceStore = true;
-	auto Visitor = [this, &bIsFromTraceStore](const TCHAR* Filename, const FFileStatData& StatData)
-	{
-		if (FPaths::GetExtension(Filename) == TEXT("utrace"))
-		{
-			TSharedPtr<FTraceFileInfo> TraceInfo = MakeShared<FTraceFileInfo>();
-			TraceInfo->FilePath = Filename;
-			TraceInfo->ModifiedTime = StatData.ModificationTime;
-			TraceInfo->bIsFromTraceStore = bIsFromTraceStore;
-
-			Traces.Add(TraceInfo);
-		}
-
-		return true;
-	};
+	FPlatformProcess::ExploreFolder(*TraceInfo->FilePath);
 
-	IFileManager::Get().IterateDirectoryStat(*StorePath, Visitor);
-	
-	bIsFromTraceStore = false;
-	IFileManager::Get().IterateDirectoryStat(*FPaths::ProfilingDir(), Visitor);
-	
-	Algo::SortBy(Traces, [](TSharedPtr<FTraceFileInfo> TraceInfo) { return *TraceInfo; });
+	return FReply::Handled();
 }
 
 #undef LOCTEXT_NAMESPACE
diff --git a/Engine/Plugins/TraceUtilities/Source/EditorTraceUtilities/Private/InsightsStatusBar/SRecentTracesList.h b/Engine/Plugins/TraceUtilities/Source/EditorTraceUtilities/Private/InsightsStatusBar/SRecentTracesList.h
--- a/Engine/Plugins/TraceUtilities/Source/EditorTraceUtilities/Private/InsightsStatusBar/SRecentTracesList.h
+++ b/Engine/Plugins/TraceUtilities/Source/EditorTraceUtilities/Private/InsightsStatusBar/SRecentTracesList.h
@@ -5,37 +5,22 @@
 #include "Widgets/SCompoundWidget.h"
 #include "Widgets/Views/SListView.h"
 
-struct FTraceFileInfo
-{
-	FString FilePath;
-	FDateTime ModifiedTime;
-	bool bIsFromTraceStore;
-
-	bool operator <(const FTraceFileInfo& rhs)
-	{
-		return this->ModifiedTime > rhs.ModifiedTime;
-	}
-};
+struct FTraceFileInfo;
 
 /**
- *  Shows a list of recent traces in a menu.
+ *  A widget that shows an entry in the recent traces submenu.
  */
-class SRecentTracesList : public SCompoundWidget
+class SRecentTracesListEntry : public SCompoundWidget
 {
-	SLATE_BEGIN_ARGS(SRecentTracesList) {}
+	SLATE_BEGIN_ARGS(SRecentTracesListEntry) {}
 	SLATE_END_ARGS()
 
-	void Construct(const FArguments& InArgs, const FString& InStorePath);
-
-	TSharedRef<ITableRow> OnGenerateRow(TSharedPtr<FTraceFileInfo> Item, const TSharedRef<STableViewBase>& OwnerTable);
+	void Construct(const FArguments& InArgs, TSharedPtr<FTraceFileInfo> InTrace, const FString& InStorePath);
 
 private:
-	void PopulateRecentTracesList();
+	FReply OpenContainingFolder();
 
 private:
-	TSharedPtr<SListView<TSharedPtr<FTraceFileInfo>>> ListView;
-	
-	TArray<TSharedPtr<FTraceFileInfo>> Traces;
-
 	FString StorePath;
+	TSharedPtr<FTraceFileInfo> TraceInfo;
 };
