Limit size of recent traces list and fix small bug.

diff --git a/Engine/Plugins/TraceUtilities/Source/EditorTraceUtilities/Private/InsightsStatusBar/SInsightsStatusBar.cpp b/Engine/Plugins/TraceUtilities/Source/EditorTraceUtilities/Private/InsightsStatusBar/SInsightsStatusBar.cpp
--- a/Engine/Plugins/TraceUtilities/Source/EditorTraceUtilities/Private/InsightsStatusBar/SInsightsStatusBar.cpp
+++ b/Engine/Plugins/TraceUtilities/Source/EditorTraceUtilities/Private/InsightsStatusBar/SInsightsStatusBar.cpp
@@ -912,6 +912,12 @@ void SInsightsStatusBarWidget::PopulateRecentTracesList()
 	IFileManager::Get().IterateDirectoryStat(*FPaths::ProfilingDir(), Visitor);
 
 	Algo::SortBy(Traces, [](TSharedPtr<FTraceFileInfo> TraceInfo) { return *TraceInfo; });
+
+	constexpr int32 MaxRecentTraces = 15;
+	if (Traces.Num() > MaxRecentTraces)
+	{
+		Traces.RemoveAt(MaxRecentTraces - 1, Traces.Num() - MaxRecentTraces);
+	}
 }
 
 void SInsightsStatusBarWidget::OpenTrace(int32 Index)
diff --git a/Engine/Plugins/TraceUtilities/Source/EditorTraceUtilities/Private/InsightsStatusBar/SRecentTracesList.cpp b/Engine/Plugins/TraceUtilities/Source/EditorTraceUtilities/Private/InsightsStatusBar/SRecentTracesList.cpp
--- a/Engine/Plugins/TraceUtilities/Source/EditorTraceUtilities/Private/InsightsStatusBar/SRecentTracesList.cpp
+++ b/Engine/Plugins/TraceUtilities/Source/EditorTraceUtilities/Private/InsightsStatusBar/SRecentTracesList.cpp
@@ -85,7 +85,7 @@ void SRecentTracesListEntry::Construct(const FArguments& InArgs, TSharedPtr<FTra
 				.ButtonStyle(FAppStyle::Get(), "SimpleButton")
 				.OnClicked(this, &SRecentTracesListEntry::OpenContainingFolder)
 				.ToolTipText(LOCTEXT("ExploreFolderTooltip", "Open the folder containing the trace file."))
-				.Visibility_Lambda([this]() {return this->IsHovered() ? EVisibility::Visible : EVisibility::Collapsed; })
+				.Visibility_Lambda([this]() {return this->IsHovered() ? EVisibility::Visible : EVisibility::Hidden; })
 				.Content()
 				[
 					SNew(SImage)
