Fixes for the Insights Status Bar Widget on Mac.

diff --git a/Engine/Plugins/TraceUtilities/Source/EditorTraceUtilities/Private/InsightsStatusBar/SInsightsStatusBar.cpp b/Engine/Plugins/TraceUtilities/Source/EditorTraceUtilities/Private/InsightsStatusBar/SInsightsStatusBar.cpp
--- a/Engine/Plugins/TraceUtilities/Source/EditorTraceUtilities/Private/InsightsStatusBar/SInsightsStatusBar.cpp
+++ b/Engine/Plugins/TraceUtilities/Source/EditorTraceUtilities/Private/InsightsStatusBar/SInsightsStatusBar.cpp
@@ -12,9 +12,10 @@
 #include "Framework/Notifications/NotificationManager.h"
 #include "GenericPlatform/GenericPlatformFile.h"
 #include "HAL/FileManager.h"
-#include "Trace/Detail/Channel.h"
-#include "Trace/Trace.h"
+#include "Logging/MessageLog.h"
+#include "Logging/TokenizedMessage.h"
 #include "Math/Color.h"
+#include "MessageLogModule.h"
 #include "Misc/ConfigCacheIni.h"
 #include "Misc/DateTime.h"
 #include "Misc/Paths.h"
@@ -24,7 +25,9 @@
 #include "SRecentTracesList.h"
 #include "Styling/StyleColors.h"
 #include "ToolMenus.h"
+#include "Trace/Detail/Channel.h"
 #include "Trace/StoreClient.h"
+#include "Trace/Trace.h"
 #include "UnrealInsightsLauncher.h"
 #include "Widgets/SBoxPanel.h"
 #include "Widgets/SOverlay.h"
@@ -226,6 +229,13 @@ void SInsightsStatusBarWidget::Construct(const FArguments& InArgs)
 		TraceDestination = ETraceDestination::File;
 	}
 
+	LogListingName = TEXT("UnrealInsights");
+	FMessageLogModule& MessageLogModule = FModuleManager::LoadModuleChecked<FMessageLogModule>("MessageLog");
+	if (!MessageLogModule.IsRegisteredLogListing(LogListingName))
+	{
+		MessageLogModule.RegisterLogListing(LogListingName, LOCTEXT("UnrealInsights", "Unreal Insights"));
+	}
+
 	FTraceAuxiliary::OnTraceStarted.AddSP(this, &SInsightsStatusBarWidget::OnTraceStarted);
 	FTraceAuxiliary::OnTraceStopped.AddSP(this, &SInsightsStatusBarWidget::OnTraceStopped);
 	FTraceAuxiliary::OnSnapshotSaved.AddSP(this, &SInsightsStatusBarWidget::OnSnapshotSaved);
@@ -330,6 +340,7 @@ TSharedRef<SWidget> SInsightsStatusBarWidget::MakeTraceMenu()
 
 	MenuBuilder.BeginSection("Options", LOCTEXT("TraceMenu_Section_Options", "Options"));
 	{
+#if PLATFORM_WINDOWS
 		MenuBuilder.AddMenuEntry(
 			LOCTEXT("OpenLiveSesssionOnTraceStart", "Open Live Session on Trace Start"),
 			LOCTEXT("OpenLiveSesssionOnTraceStartDesc", "When set, the live session will be automatically opened in Unreal Insights when tracing is started.\nThis option will only apply when tracing to the trace store."),
@@ -351,6 +362,7 @@ TSharedRef<SWidget> SInsightsStatusBarWidget::MakeTraceMenu()
 			NAME_None,
 			EUserInterfaceActionType::ToggleButton
 		);
+#endif
 
 		MenuBuilder.AddMenuEntry(
 			LOCTEXT("ShowInExplorerAfterTrace", "Show in Explorer after Trace"),
@@ -398,6 +410,7 @@ TSharedRef<SWidget> SInsightsStatusBarWidget::MakeTraceMenu()
 			EUserInterfaceActionType::Button
 		);
 
+#if PLATFORM_WINDOWS
 		MenuBuilder.AddMenuEntry(
 			LOCTEXT("OpenLiveSessionLabel", "Open Live Session"),
 			LOCTEXT("OpenLiveSessionTooltip", "Opening the live session is possible only while tracing to the trace store."),
@@ -407,6 +420,7 @@ TSharedRef<SWidget> SInsightsStatusBarWidget::MakeTraceMenu()
 			NAME_None,
 			EUserInterfaceActionType::Button
 		);
+#endif
 
 		MenuBuilder.AddSubMenu
 		(
@@ -483,6 +497,10 @@ FText SInsightsStatusBarWidget::GetTitleToolTipText() const
 	{
 		DescBuilder.AppendLineFormat(LOCTEXT("TracingToText", "Tracing to: {0}"), FText::FromString(Dest));
 	}
+	else if (UE::Trace::IsTracing())
+	{
+		DescBuilder.AppendLine(LOCTEXT("TracingToUnknownText", "Tracing to unknown target (externally set)."));
+	}
 	else
 	{
 		DescBuilder.AppendLine(LOCTEXT("NotTracingText","Not currently tracing."));
@@ -493,7 +511,7 @@ FText SInsightsStatusBarWidget::GetTitleToolTipText() const
 
 FSlateColor SInsightsStatusBarWidget::GetRecordingButtonColor() const
 {
-	if (!FTraceAuxiliary::IsConnected())
+	if (!UE::Trace::IsTracing())
 	{
 		return FStyleColors::White;
 	}
@@ -503,7 +521,7 @@ FSlateColor SInsightsStatusBarWidget::GetRecordingButtonColor() const
 
 FSlateColor SInsightsStatusBarWidget::GetRecordingButtonOutlineColor() const
 {
-	if (!FTraceAuxiliary::IsConnected())
+	if (!UE::Trace::IsTracing())
 	{
 		ConnectionStartTime = FSlateApplication::Get().GetCurrentTime();
 		return FLinearColor::White.CopyWithNewOpacity(0.5f);
@@ -515,7 +533,7 @@ FSlateColor SInsightsStatusBarWidget::GetRecordingButtonOutlineColor() const
 
 FText SInsightsStatusBarWidget::GetRecordingButtonTooltipText() const
 {
-	if (!FTraceAuxiliary::IsConnected())
+	if (!UE::Trace::IsTracing())
 	{
 		return LOCTEXT("StartTracing", "Start tracing. The trace destination is set from the menu.");
 	}
@@ -523,31 +541,23 @@ FText SInsightsStatusBarWidget::GetRecordingButtonTooltipText() const
 	return LOCTEXT("StopTracing", "Stop Tracing.");
 }
 
-void SInsightsStatusBarWidget::SendSnapshotNotification()
+void SInsightsStatusBarWidget::LogMessage(const FText& Text)
 {
-	FNotificationInfo Info(LOCTEXT("SnapshotSaved", "Insights Snapshot saved."));
-	Info.bFireAndForget = true;
-	Info.FadeOutDuration = 1.0f;
-	Info.ExpireDuration = 4.0f;
-
-	Info.Text = LOCTEXT("SnapshotSavedHeading", "Insights Snapshot saved.");
-	Info.SubText = LOCTEXT("SnapshotSavedText", "A snapshot .utrace with the most recent events has been saved to your Saved/Profiling/ directory.");
-
-	FSlateNotificationManager::Get().AddNotification(Info);
+	FMessageLog ReportMessageLog(LogListingName);
+	TSharedRef<FTokenizedMessage> TokenizedMessage = FTokenizedMessage::Create(EMessageSeverity::Error, Text);
+	ReportMessageLog.AddMessage(TokenizedMessage);
+	ReportMessageLog.Notify();
 }
 
-void SInsightsStatusBarWidget::SendTraceStartedNotification()
+void SInsightsStatusBarWidget::ShowNotification(const FText& Text, const FText& SubText)
 {
-	FNotificationInfo Info(LOCTEXT("TracingStarted", "Tracing started."));
+	FNotificationInfo Info(Text);
 	Info.bFireAndForget = true;
 	Info.FadeOutDuration = 1.0f;
 	Info.ExpireDuration = 4.0f;
 
-	Info.Text = LOCTEXT("TracingStartedHeading", "Tracing started.");
-	Info.SubText = FText::Format(
-		LOCTEXT("TracingStartedText", "Trace is now active and saving to the following location (file or tracestore):\n{0}")
-		, FText::FromString(FTraceAuxiliary::GetTraceDestination())
-	);
+	Info.SubText = SubText;
+
 	FSlateNotificationManager::Get().AddNotification(Info);
 }
 
@@ -629,7 +639,7 @@ bool SInsightsStatusBarWidget::SetTraceDestination_IsChecked(ETraceDestination I
 
 bool SInsightsStatusBarWidget::SetTraceDestination_CanExecute()
 {
-	if (!FTraceAuxiliary::IsConnected())
+	if (!UE::Trace::IsTracing())
 	{
 		return true;
 	}
@@ -639,8 +649,15 @@ bool SInsightsStatusBarWidget::SetTraceDestination_CanExecute()
 
 void SInsightsStatusBarWidget::SaveSnapshot()
 {
-	FTraceAuxiliary::WriteSnapshot(nullptr);
-	SendSnapshotNotification();
+	bool bResult = FTraceAuxiliary::WriteSnapshot(nullptr);
+	if (bResult)
+	{
+		ShowNotification(LOCTEXT("SnapshotSavedHeading", "Insights Snapshot saved."), LOCTEXT("SnapshotSavedText", "A snapshot .utrace with the most recent events has been saved to your Saved/Profiling/ directory."));
+	}
+	else
+	{
+		LogMessage(LOCTEXT("SnapshotSavedError", "The snapshot could not be saved."));
+	}
 }
 
 bool SInsightsStatusBarWidget::SaveSnapshot_CanExecute()
@@ -650,7 +667,7 @@ bool SInsightsStatusBarWidget::SaveSnapshot_CanExecute()
 
 FText SInsightsStatusBarWidget::GetTraceMenuItemText() const
 {
-	if (FTraceAuxiliary::IsConnected())
+	if (UE::Trace::IsTracing())
 	{
 		return LOCTEXT("StopTraceButtonText", "Stop Trace");
 	}
@@ -660,7 +677,7 @@ FText SInsightsStatusBarWidget::GetTraceMenuItemText() const
 
 FText SInsightsStatusBarWidget::GetTraceMenuItemTooltipText() const
 {
-	if (FTraceAuxiliary::IsConnected())
+	if (UE::Trace::IsTracing())
 	{
 		return LOCTEXT("StopTraceButtonTooltip", "Stop tracing");
 	}
@@ -670,27 +687,46 @@ FText SInsightsStatusBarWidget::GetTraceMenuItemTooltipText() const
 
 void SInsightsStatusBarWidget::ToggleTracing_OnClicked()
 {
-	if (FTraceAuxiliary::IsConnected())
+	if (UE::Trace::IsTracing())
 	{
-		FTraceAuxiliary::Stop();
+		bool bResult = FTraceAuxiliary::Stop();
+		if (!bResult)
+		{
+			LogMessage(LOCTEXT("TraceStopFailedMsg", "There was no trace connection to stop."));
+		}
 	}
 	else
 	{
-		StartTracing();
-		SendTraceStartedNotification();
+		bool bResult = StartTracing();
+		if (bResult)
+		{
+			FString TraceDestinationStr = FTraceAuxiliary::GetTraceDestination();
+			if (TraceDestinationStr.IsEmpty())
+			{
+				TraceDestinationStr = TEXT("External Target");
+			}
+
+			ShowNotification(LOCTEXT("TraceMsg", "Trace Started"), FText::Format(LOCTEXT("TracingStartedText", "Trace is now active and saving to the following location (file or tracestore):\n{0}"), FText::FromString(TraceDestinationStr)));
+		}
+		else
+		{
+			LogMessage(LOCTEXT("TraceFailedToStartMsg", "Trace Failed to Start."));
+		}
 	}
 }
 
-void SInsightsStatusBarWidget::StartTracing()
+bool SInsightsStatusBarWidget::StartTracing()
 {
 	if (TraceDestination == ETraceDestination::TraceStore)
 	{
-		FTraceAuxiliary::Start(FTraceAuxiliary::EConnectionType::Network, TEXT("localhost"), nullptr);
+		return FTraceAuxiliary::Start(FTraceAuxiliary::EConnectionType::Network, TEXT("localhost"), nullptr);
 	}
 	else if (TraceDestination == ETraceDestination::File)
 	{
-		FTraceAuxiliary::Start(FTraceAuxiliary::EConnectionType::File, nullptr, nullptr);
+		return FTraceAuxiliary::Start(FTraceAuxiliary::EConnectionType::File, nullptr, nullptr);
 	}
+
+	return false;
 }
 
 EVisibility SInsightsStatusBarWidget::GetStartTraceIconVisibility() const
@@ -705,7 +741,7 @@ EVisibility SInsightsStatusBarWidget::GetStartTraceIconVisibility() const
 
 EVisibility SInsightsStatusBarWidget::GetStopTraceIconVisibility() const
 {
-	if (bIsTraceRecordButtonHovered && FTraceAuxiliary::IsConnected())
+	if (bIsTraceRecordButtonHovered && UE::Trace::IsTracing())
 	{
 		return EVisibility::Visible;
 	}
@@ -770,22 +806,26 @@ void SInsightsStatusBarWidget::CacheTraceStorePath()
 {
 	if (TraceStorePath.IsEmpty())
 	{
+#if PLATFORM_WINDOWS
 		UE::Trace::FStoreClient* StoreClient = UE::Trace::FStoreClient::Connect(TEXT("localhost"));
 
 		if (!StoreClient)
 		{
-			// TODO: Add Error Message
+			LogMessage(LOCTEXT("FailedConnectionToStoreMsg", "Failed to connect to the store client."));
 			return;
 		}
 
 		const UE::Trace::FStoreClient::FStatus* Status = StoreClient->GetStatus();
 		if (!Status)
 		{
-			// TODO: Add Error Message
+			LogMessage(LOCTEXT("FailedToGetStoreStatusMsg", "Failed to get the status of the store client."));
 			return;
 		}
 
 		TraceStorePath = FString(Status->GetStoreDir());
+#else
+		TraceStorePath = FPaths::Combine(FPaths::EngineDir(), TEXT("/Programs/UnrealInsights/Saved/TraceSessions/"));
+#endif
 	}
 }
 
diff --git a/Engine/Plugins/TraceUtilities/Source/EditorTraceUtilities/Private/InsightsStatusBar/SInsightsStatusBar.h b/Engine/Plugins/TraceUtilities/Source/EditorTraceUtilities/Private/InsightsStatusBar/SInsightsStatusBar.h
--- a/Engine/Plugins/TraceUtilities/Source/EditorTraceUtilities/Private/InsightsStatusBar/SInsightsStatusBar.h
+++ b/Engine/Plugins/TraceUtilities/Source/EditorTraceUtilities/Private/InsightsStatusBar/SInsightsStatusBar.h
@@ -85,14 +85,14 @@ private:
 	EVisibility GetStartTraceIconVisibility() const;
 	EVisibility GetStopTraceIconVisibility() const;
 
-	void StartTracing();
+	bool StartTracing();
 
 	TSharedRef<SWidget> MakeTraceMenu();
 	void Channels_BuildMenu(FMenuBuilder& MenuBuilder);
 	void Traces_BuildMenu(FMenuBuilder& MenuBuilder);
 
-	static void SendSnapshotNotification();
-	static void SendTraceStartedNotification();
+	void LogMessage(const FText& Text);
+	void ShowNotification(const FText& Text, const FText& SubText);
 
 	void SetTraceChannels(const TCHAR* InChannels);
 	bool IsPresetSet(const TCHAR* InChannels) const;
@@ -147,4 +147,5 @@ private:
 	TSharedPtr<FUICommandList> CommandList;
 
 	TArray<TSharedPtr<FTraceFileInfo>> Traces;
+	FName LogListingName;
 };
diff --git a/Engine/Plugins/TraceUtilities/Source/EditorTraceUtilities/Private/InsightsStatusBar/SRecentTracesList.cpp b/Engine/Plugins/TraceUtilities/Source/EditorTraceUtilities/Private/InsightsStatusBar/SRecentTracesList.cpp
--- a/Engine/Plugins/TraceUtilities/Source/EditorTraceUtilities/Private/InsightsStatusBar/SRecentTracesList.cpp
+++ b/Engine/Plugins/TraceUtilities/Source/EditorTraceUtilities/Private/InsightsStatusBar/SRecentTracesList.cpp
@@ -98,7 +98,8 @@ void SRecentTracesListEntry::Construct(const FArguments& InArgs, TSharedPtr<FTra
 
 FReply SRecentTracesListEntry::OpenContainingFolder()
 {
-	FPlatformProcess::ExploreFolder(*TraceInfo->FilePath);
+	FString FullPath(FPaths::ConvertRelativePathToFull(TraceInfo->FilePath));
+	FPlatformProcess::ExploreFolder(*FullPath);
 
 	return FReply::Handled();
 }
