Remove Insights menu items from the Tools menu as the functionality has been moved to the Insights Status Bar Widget.

diff --git a/Engine/Plugins/TraceUtilities/Source/EditorTraceUtilities/Private/EditorTraceUtilities.cpp b/Engine/Plugins/TraceUtilities/Source/EditorTraceUtilities/Private/EditorTraceUtilities.cpp
--- a/Engine/Plugins/TraceUtilities/Source/EditorTraceUtilities/Private/EditorTraceUtilities.cpp
+++ b/Engine/Plugins/TraceUtilities/Source/EditorTraceUtilities/Private/EditorTraceUtilities.cpp
@@ -4,7 +4,6 @@
 
 #include "Misc/ConfigContext.h"
 #include "ToolMenus.h"
-#include "UnrealInsightsLauncher.h"
 
 DEFINE_LOG_CATEGORY(LogTraceUtilities)
 
@@ -20,10 +19,6 @@ FString FEditorTraceUtilitiesModule::EditorTraceUtilitiesIni;
 void FEditorTraceUtilitiesModule::StartupModule()
 {
 	LLM_SCOPE_BYNAME(TEXT("Insights"));
-	RegisterStartupCallbackHandle = UToolMenus::RegisterStartupCallback(
-		FSimpleMulticastDelegate::FDelegate::CreateSP(
-			FUnrealInsightsLauncher::Get().ToSharedRef(),
-			&FUnrealInsightsLauncher::RegisterMenus));
 
 	RegisterInsightsStatusWidgetWithToolMenu();
 
diff --git a/Engine/Plugins/TraceUtilities/Source/EditorTraceUtilities/Private/UnrealInsightsLauncher.cpp b/Engine/Plugins/TraceUtilities/Source/EditorTraceUtilities/Private/UnrealInsightsLauncher.cpp
--- a/Engine/Plugins/TraceUtilities/Source/EditorTraceUtilities/Private/UnrealInsightsLauncher.cpp
+++ b/Engine/Plugins/TraceUtilities/Source/EditorTraceUtilities/Private/UnrealInsightsLauncher.cpp
@@ -49,47 +49,12 @@ FUnrealInsightsLauncher::~FUnrealInsightsLauncher()
 
 }
 	
-void FUnrealInsightsLauncher::RegisterMenus()
-{
-	FToolMenuOwnerScoped OwnerScoped(this);
-
-	UToolMenu* ProfileMenu = UToolMenus::Get()->ExtendMenu("MainFrame.MainMenu.Tools");
-	if (ProfileMenu)
-	{
-		FToolMenuSection& Section = ProfileMenu->AddSection("Unreal Insights", FText::FromString(TEXT("Unreal Insights")));
-		Section.AddMenuEntry("OpenUnrealInsights",
-			LOCTEXT("OpenUnrealInsights_Label", "Run Unreal Insights"),
-			LOCTEXT("OpenUnrealInsights_Desc", "Run the Unreal Insights standalone application."),
-			FSlateIcon(FAppStyle::GetAppStyleSetName(), "UnrealInsights.MenuIcon"),
-			FUIAction(FExecuteAction::CreateRaw(this, &FUnrealInsightsLauncher::RunUnrealInsights_Execute), FCanExecuteAction())
-		);
-		Section.AddMenuEntry("OpenLiveTrace",
-			LOCTEXT("OpenLiveTrace_Label", "Open active Trace"),
-			LOCTEXT("OpenLiveTrace_Desc", "Opens the currently running trace in Unreal Insights."),
-			FSlateIcon(FAppStyle::GetAppStyleSetName(), "UnrealInsights.MenuIcon"),
-			FUIAction(FExecuteAction::CreateRaw(this, &FUnrealInsightsLauncher::RunUnrealInsights_OpenLiveTrace), FCanExecuteAction::CreateLambda(
-				[](){ return FTraceAuxiliary::IsConnected();}))
-		);
-	}
-}
-
 FString FUnrealInsightsLauncher::GetInsightsApplicationPath()
 {
 	FString Path = FPlatformProcess::GenerateApplicationPath(TEXT("UnrealInsights"), EBuildConfiguration::Development);
 	return FPaths::ConvertRelativePathToFull(Path);
 }
 
-void FUnrealInsightsLauncher::RunUnrealInsights_Execute()
-{
-	FString Path = GetInsightsApplicationPath();
-	StartUnrealInsights(Path);
-}
-
-void FUnrealInsightsLauncher::RunUnrealInsights_OpenLiveTrace()
-{
-	TryOpenTraceFromDestination(FTraceAuxiliary::GetTraceDestination());
-}
-
 void FUnrealInsightsLauncher::StartUnrealInsights(const FString& Path, const FString& Parameters)
 {
     if (!FPaths::FileExists(Path))
diff --git a/Engine/Plugins/TraceUtilities/Source/EditorTraceUtilities/Private/UnrealInsightsLauncher.h b/Engine/Plugins/TraceUtilities/Source/EditorTraceUtilities/Private/UnrealInsightsLauncher.h
--- a/Engine/Plugins/TraceUtilities/Source/EditorTraceUtilities/Private/UnrealInsightsLauncher.h
+++ b/Engine/Plugins/TraceUtilities/Source/EditorTraceUtilities/Private/UnrealInsightsLauncher.h
@@ -23,9 +23,6 @@ public:
 		return Instance;
 	}
 
-	
-	void RegisterMenus();
-
 	/**
 	 * Returns the full path to UnrealInsights.exe for the current engine installation 
 	 */
@@ -75,14 +72,6 @@ private:
 	 */
 	void TryBuildUnrealInsightsExe(const FString& Path, const FString& LaunchParameters = TEXT(""));
 
-	/// callback from Options Menu Entry "Start Unreal Insights"
-	void RunUnrealInsights_Execute();
-
-	/**
-	 * Opens the currently live trace if connected.
-	 */
-	 void RunUnrealInsights_OpenLiveTrace();
-
 	/// Logs an error message to the MessageLog window
 	void LogMessage(const FText& Message);
 
