Properly implement snapshot to server in Insights status bar

diff --git a/Engine/Plugins/TraceUtilities/Source/EditorTraceUtilities/Private/InsightsStatusBar/SInsightsStatusBar.cpp b/Engine/Plugins/TraceUtilities/Source/EditorTraceUtilities/Private/InsightsStatusBar/SInsightsStatusBar.cpp
--- a/Engine/Plugins/TraceUtilities/Source/EditorTraceUtilities/Private/InsightsStatusBar/SInsightsStatusBar.cpp
+++ b/Engine/Plugins/TraceUtilities/Source/EditorTraceUtilities/Private/InsightsStatusBar/SInsightsStatusBar.cpp
@@ -215,7 +215,7 @@ void SInsightsStatusBarWidget::Construct(const FArguments& InArgs)
 				SNew(SImage)
 				.DesiredSizeOverride(CoreStyleConstants::Icon16x16)
 				.Image(FEditorTraceUtilitiesStyle::Get().GetBrush("Icons.TraceSnapshot.StatusBar"))
-				.ToolTipText(LOCTEXT("SaveSnapShot","Snapshot: Save Current Trace Buffer To File"))
+				.ToolTipText(LOCTEXT("SaveSnapShot","Snapshot: Save Current Trace Buffer to active destination."))
 			]
 		]
 	];
@@ -649,15 +649,25 @@ bool SInsightsStatusBarWidget::SetTraceDestination_CanExecute()
 
 void SInsightsStatusBarWidget::SaveSnapshot()
 {
-	bool bResult = FTraceAuxiliary::WriteSnapshot(nullptr);
-	if (bResult)
+	if (TraceDestination == ETraceDestination::File)
 	{
-		ShowNotification(LOCTEXT("SnapshotSavedHeading", "Insights Snapshot saved."), LOCTEXT("SnapshotSavedText", "A snapshot .utrace with the most recent events has been saved to your Saved/Profiling/ directory."));
+		const bool bResult = FTraceAuxiliary::WriteSnapshot(nullptr);
+		if (bResult)
+		{
+			ShowNotification(LOCTEXT("SnapshotSavedHeading", "Insights Snapshot saved."), LOCTEXT("SnapshotSavedFileText", "A snapshot .utrace with the most recent events has been saved to your Saved/Profiling/ directory."));
+			return;
+		}
 	}
 	else
 	{
-		LogMessage(LOCTEXT("SnapshotSavedError", "The snapshot could not be saved."));
+		const bool bResult = FTraceAuxiliary::SendSnapshot(nullptr);
+		if (bResult)
+		{
+			ShowNotification(LOCTEXT("SnapshotSavedHeading", "Insights Snapshot saved."), LOCTEXT("SnapshotSavedServerText", "A snapshot .utrace with the most recent events has been saved to your the trace server."));
+			return;
+		}
 	}
+	LogMessage(LOCTEXT("SnapshotSavedError", "The snapshot could not be saved."));
 }
 
 bool SInsightsStatusBarWidget::SaveSnapshot_CanExecute()
diff --git a/Engine/Source/Runtime/Core/Private/ProfilingDebugging/TraceAuxiliary.cpp b/Engine/Source/Runtime/Core/Private/ProfilingDebugging/TraceAuxiliary.cpp
--- a/Engine/Source/Runtime/Core/Private/ProfilingDebugging/TraceAuxiliary.cpp
+++ b/Engine/Source/Runtime/Core/Private/ProfilingDebugging/TraceAuxiliary.cpp
@@ -584,6 +584,12 @@ bool FTraceAuxiliaryImpl::SendSnapshot(const TCHAR* InHost, uint32 InPort, const
 {
 	double StartTime = FPlatformTime::Seconds();
 
+	// If no host is set, assume localhost
+	if (!InHost)
+	{
+		InHost = TEXT("localhost");
+	}
+
 	UE_LOG_REF(LogCategory, Log, TEXT("Sending trace snapshot to '%s'..."), InHost);
 
 	const bool bResult = UE::Trace::SendSnapshotTo(InHost, InPort);
@@ -966,10 +972,6 @@ static void TraceAuxiliarySnapshotSend(const TArray<FString>& Args)
 	{
 		LexFromString(Port, *Args[1]);
 	}
-	if (Args.Num() == 0)
-	{
-		Host = TEXT("localhost");
-	}
 	if (Args.Num() > 2) 
 	{
 		UE_LOG(LogConsoleResponse, Warning, TEXT("Invalid arguments. Usage: Trace.SnapshotFile <Host> <Port>"));
@@ -1428,6 +1430,16 @@ bool FTraceAuxiliary::WriteSnapshot(const TCHAR* InFilePath)
 #endif
 }
 
+////////////////////////////////////////////////////////////////////////////////
+bool FTraceAuxiliary::SendSnapshot(const TCHAR* Host, uint32 Port)
+{
+#if UE_TRACE_ENABLED
+	return GTraceAuxiliary.SendSnapshot(Host, Port, LogCore);
+#else
+	return true;
+#endif
+}
+
 ////////////////////////////////////////////////////////////////////////////////
 #if UE_TRACE_ENABLED
 static void TraceAuxiliaryAddPostForkCallback(const TCHAR* CommandLine)
diff --git a/Engine/Source/Runtime/Core/Public/ProfilingDebugging/TraceAuxiliary.h b/Engine/Source/Runtime/Core/Public/ProfilingDebugging/TraceAuxiliary.h
--- a/Engine/Source/Runtime/Core/Public/ProfilingDebugging/TraceAuxiliary.h
+++ b/Engine/Source/Runtime/Core/Public/ProfilingDebugging/TraceAuxiliary.h
@@ -104,6 +104,12 @@ public:
 	 * @param FilePath Path to the file to write the snapshot to. If it is null or empty a file path will be generated.
 	 */
 	static bool WriteSnapshot(const TCHAR* FilePath);
+	
+	/**
+	 * Write tailing memory state to a trace server.
+	 * @param FilePath Path to the file to write the snapshot to. If it is null or empty a file path will be generated.
+	 */
+	static bool SendSnapshot(const TCHAR* Host = nullptr, uint32 Port = 0);
 
 	/**
 	 * Initialize Trace systems.
