Fix snapshot bugs in Insights Status Bar Widget.

diff --git a/Engine/Plugins/TraceUtilities/Source/EditorTraceUtilities/Private/InsightsStatusBar/SInsightsStatusBar.cpp b/Engine/Plugins/TraceUtilities/Source/EditorTraceUtilities/Private/InsightsStatusBar/SInsightsStatusBar.cpp
--- a/Engine/Plugins/TraceUtilities/Source/EditorTraceUtilities/Private/InsightsStatusBar/SInsightsStatusBar.cpp
+++ b/Engine/Plugins/TraceUtilities/Source/EditorTraceUtilities/Private/InsightsStatusBar/SInsightsStatusBar.cpp
@@ -64,7 +64,6 @@ public:
 	}
 };
 
-
 ////////////////////////////////////////////////////////////////////////////////////////////////////
 // FInsightsStatusBarWidgetCommands
 ////////////////////////////////////////////////////////////////////////////////////////////////////
@@ -328,7 +327,7 @@ TSharedRef<SWidget> SInsightsStatusBarWidget::MakeTraceMenu()
 
 		MenuBuilder.AddMenuEntry(
 			LOCTEXT("SaveSnapshotLabel", "Save Trace Snapshot"),
-			LOCTEXT("SaveSnapshotTooltip", "Save the current trace buffer to file."),
+			LOCTEXT("SaveSnapshotTooltip", "Save the current trace buffer to the selected trace destination."),
 			FSlateIcon(FEditorTraceUtilitiesStyle::Get().GetStyleSetName(), "Icons.TraceSnapshot.Menu"),
 			FUIAction(FExecuteAction::CreateSP(this, &SInsightsStatusBarWidget::SaveSnapshot),
 					  FCanExecuteAction::CreateSP(this, &SInsightsStatusBarWidget::SaveSnapshot_CanExecute)),
@@ -589,42 +588,65 @@ void SInsightsStatusBarWidget::OpenProfilingDirectory()
 
 void SInsightsStatusBarWidget::OpenTraceStoreDirectory_OnClicked()
 {
-	OpenTraceStoreDirectory(false);
+	OpenTraceStoreDirectory(ESelectLatestTraceCriteria::None);
 }
 
-void SInsightsStatusBarWidget::OpenTraceStoreDirectory(bool bSelectLastTrace)
+void SInsightsStatusBarWidget::OpenTraceStoreDirectory(ESelectLatestTraceCriteria Criteria)
 {
 	CacheTraceStorePath();
+	FString Path = TraceStorePath;
+
+	if (Criteria != ESelectLatestTraceCriteria::None)
+	{
+		Path = GetLatestTraceFileFromFolder(TraceStorePath, Criteria);
+	}
+
+	FPlatformProcess::ExploreFolder(*Path);
+}
+
+void SInsightsStatusBarWidget::OpenLatestTraceFromFolder(const FString& InFolder, ESelectLatestTraceCriteria InCriteria)
+{
+	FString Path = GetLatestTraceFileFromFolder(InFolder, InCriteria);
 
-	FString StorePathLocal = TraceStorePath;
-	if (bSelectLastTrace)
+	if (!Path.IsEmpty())
 	{
-		FString MostRecentTraceName;
-		FDateTime LatestDateTime;
+		FUnrealInsightsLauncher::Get()->TryOpenTraceFromDestination(Path);
+	}
+}
 
-		auto Visitor = [&MostRecentTraceName, &LatestDateTime](const TCHAR* Filename, const FFileStatData& StatData)
+FString SInsightsStatusBarWidget::GetLatestTraceFileFromFolder(const FString& InFolder, ESelectLatestTraceCriteria InCriteria)
+{
+	FString Result;
+	FString MostRecentTraceName;
+	FDateTime LatestDateTime;
+
+	auto Visitor = [&MostRecentTraceName, &LatestDateTime, InCriteria](const TCHAR* Filename, const FFileStatData& StatData)
+	{
+		if (FPaths::GetExtension(Filename) == TEXT("utrace"))
 		{
-			if (FPaths::GetExtension(Filename) == TEXT("utrace"))
+			if (InCriteria == ESelectLatestTraceCriteria::ModifiedTime && LatestDateTime < StatData.ModificationTime)
 			{
-				if (LatestDateTime < StatData.ModificationTime)
-				{
-					LatestDateTime = StatData.ModificationTime;
-					MostRecentTraceName = Filename;
-				}
+				LatestDateTime = StatData.ModificationTime;
+				MostRecentTraceName = Filename;
 			}
-			return true;
-		};
-
-		IFileManager::Get().IterateDirectoryStat(*StorePathLocal, Visitor);
 
-		if (!MostRecentTraceName.IsEmpty())
-		{
-			StorePathLocal = MostRecentTraceName;
+			if (InCriteria == ESelectLatestTraceCriteria::CreatedTime && LatestDateTime < StatData.CreationTime)
+			{
+				LatestDateTime = StatData.CreationTime;
+				MostRecentTraceName = Filename;
+			}
 		}
+		return true;
+	};
+
+	IFileManager::Get().IterateDirectoryStat(*InFolder, Visitor);
+
+	if (!MostRecentTraceName.IsEmpty())
+	{
+		Result = FPaths::ConvertRelativePathToFull(MostRecentTraceName);
 	}
 
-	FString FullPath(FPaths::ConvertRelativePathToFull(StorePathLocal));
-	FPlatformProcess::ExploreFolder(*FullPath);
+	return Result;
 }
 
 void SInsightsStatusBarWidget::SetTraceDestination_Execute(ETraceDestination InDestination)
@@ -663,7 +685,7 @@ void SInsightsStatusBarWidget::SaveSnapshot()
 		const bool bResult = FTraceAuxiliary::SendSnapshot(nullptr);
 		if (bResult)
 		{
-			ShowNotification(LOCTEXT("SnapshotSavedHeading", "Insights Snapshot saved."), LOCTEXT("SnapshotSavedServerText", "A snapshot .utrace with the most recent events has been saved to your the trace server."));
+			ShowNotification(LOCTEXT("SnapshotSavedHeading", "Insights Snapshot saved."), LOCTEXT("SnapshotSavedServerText", "A snapshot .utrace with the most recent events has been saved to your trace server."));
 			return;
 		}
 	}
@@ -791,7 +813,7 @@ void SInsightsStatusBarWidget::OnTraceStopped(FTraceAuxiliary::EConnectionType I
 	{
 		if (InTraceType == FTraceAuxiliary::EConnectionType::Network)
 		{
-			OpenTraceStoreDirectory(true);
+			OpenTraceStoreDirectory(ESelectLatestTraceCriteria::ModifiedTime);
 		}
 		else if (InTraceType == FTraceAuxiliary::EConnectionType::File)
 		{
@@ -800,15 +822,30 @@ void SInsightsStatusBarWidget::OnTraceStopped(FTraceAuxiliary::EConnectionType I
 	}
 }
 
-void SInsightsStatusBarWidget::OnSnapshotSaved(const FString& InPath)
+void SInsightsStatusBarWidget::OnSnapshotSaved(FTraceAuxiliary::EConnectionType InTraceType, const FString& InTraceDestination)
 {
 	if (GetBooleanSettingValue(OpenInsightsAfterTraceSettingName))
 	{
-		FUnrealInsightsLauncher::Get()->TryOpenTraceFromDestination(InPath);
+		if (InTraceType == FTraceAuxiliary::EConnectionType::Network)
+		{
+			CacheTraceStorePath();
+			OpenLatestTraceFromFolder(TraceStorePath, ESelectLatestTraceCriteria::CreatedTime);
+		}
+		else if (InTraceType == FTraceAuxiliary::EConnectionType::File)
+		{
+			FUnrealInsightsLauncher::Get()->TryOpenTraceFromDestination(InTraceDestination);
+		}
 	}
 	if (GetBooleanSettingValue(ShowInExplorerAfterTraceSettingName))
 	{
-		FPlatformProcess::ExploreFolder(*InPath);
+		if (InTraceType == FTraceAuxiliary::EConnectionType::Network)
+		{
+			OpenTraceStoreDirectory(ESelectLatestTraceCriteria::CreatedTime);
+		}
+		else if (InTraceType == FTraceAuxiliary::EConnectionType::File)
+		{
+			FPlatformProcess::ExploreFolder(*InTraceDestination);
+		}
 	}
 }
 
diff --git a/Engine/Plugins/TraceUtilities/Source/EditorTraceUtilities/Private/InsightsStatusBar/SInsightsStatusBar.h b/Engine/Plugins/TraceUtilities/Source/EditorTraceUtilities/Private/InsightsStatusBar/SInsightsStatusBar.h
--- a/Engine/Plugins/TraceUtilities/Source/EditorTraceUtilities/Private/InsightsStatusBar/SInsightsStatusBar.h
+++ b/Engine/Plugins/TraceUtilities/Source/EditorTraceUtilities/Private/InsightsStatusBar/SInsightsStatusBar.h
@@ -13,19 +13,6 @@ class FUICommandList;
 
 TSharedRef<SWidget> CreateInsightsStatusBarWidget();
 
-enum class ETraceDestination : uint32
-{
-	TraceStore = 0,
-	File = 1
-};
-
-struct FChannelData
-{
-	FString Name;
-	bool bIsEnabled = false;
-	bool bIsReadOnly = false;
-};
-
 struct FTraceFileInfo
 {
 	FString FilePath;
@@ -44,6 +31,27 @@ struct FTraceFileInfo
  */
 class SInsightsStatusBarWidget : public SCompoundWidget
 {
+private:
+	enum class ETraceDestination : uint32
+	{
+		TraceStore = 0,
+		File = 1
+	};
+
+	struct FChannelData
+	{
+		FString Name;
+		bool bIsEnabled = false;
+		bool bIsReadOnly = false;
+	};
+
+	enum class ESelectLatestTraceCriteria : uint32
+	{
+		None,
+		CreatedTime,
+		ModifiedTime,
+	};
+
 public:
 	SLATE_BEGIN_ARGS(SInsightsStatusBarWidget) {}
 	SLATE_END_ARGS()
@@ -68,7 +76,10 @@ private:
 	void OpenProfilingDirectory();
 
 	void OpenTraceStoreDirectory_OnClicked();
-	void OpenTraceStoreDirectory(bool bSelectLastTrace);
+	void OpenTraceStoreDirectory(ESelectLatestTraceCriteria Criteria);
+
+	void OpenLatestTraceFromFolder(const FString& InFolder, ESelectLatestTraceCriteria InCriteria);
+	FString GetLatestTraceFileFromFolder(const FString& InFolder, ESelectLatestTraceCriteria InCriteria);
 
 	void SetTraceDestination_Execute(ETraceDestination InDestination);
 	bool SetTraceDestination_CanExecute();
@@ -102,7 +113,7 @@ private:
 
 	void OnTraceStarted(FTraceAuxiliary::EConnectionType TraceType, const FString& TraceDestination);
 	void OnTraceStopped(FTraceAuxiliary::EConnectionType TraceType, const FString& TraceDestination);
-	void OnSnapshotSaved(const FString& InPath);
+	void OnSnapshotSaved(FTraceAuxiliary::EConnectionType TraceType, const FString& TraceDestination);
 
 	void CacheTraceStorePath();
 
diff --git a/Engine/Source/Runtime/Core/Private/ProfilingDebugging/TraceAuxiliary.cpp b/Engine/Source/Runtime/Core/Private/ProfilingDebugging/TraceAuxiliary.cpp
--- a/Engine/Source/Runtime/Core/Private/ProfilingDebugging/TraceAuxiliary.cpp
+++ b/Engine/Source/Runtime/Core/Private/ProfilingDebugging/TraceAuxiliary.cpp
@@ -568,7 +568,7 @@ bool FTraceAuxiliaryImpl::WriteSnapshot(const TCHAR* InFilePath, const FTraceAux
 
 	if (bResult)
 	{
-		FTraceAuxiliary::OnSnapshotSaved.Broadcast(NativePath);
+		FTraceAuxiliary::OnSnapshotSaved.Broadcast(FTraceAuxiliary::EConnectionType::File, NativePath);
 		UE_LOG_REF(LogCategory, Display, TEXT("Trace snapshot generated in %.3f seconds to \"%s\"."), FPlatformTime::Seconds() - StartTime, *NativePath);
 	}
 	else
@@ -596,7 +596,7 @@ bool FTraceAuxiliaryImpl::SendSnapshot(const TCHAR* InHost, uint32 InPort, const
 
 	if (bResult)
 	{
-		FTraceAuxiliary::OnSnapshotSaved.Broadcast(InHost);
+		FTraceAuxiliary::OnSnapshotSaved.Broadcast(FTraceAuxiliary::EConnectionType::Network, InHost);
 		UE_LOG_REF(LogCategory, Display, TEXT("Trace snapshot generated in %.3f seconds to \"%s\"."), FPlatformTime::Seconds() - StartTime, InHost);
 	}
 	else
diff --git a/Engine/Source/Runtime/Core/Public/ProfilingDebugging/TraceAuxiliary.h b/Engine/Source/Runtime/Core/Public/ProfilingDebugging/TraceAuxiliary.h
--- a/Engine/Source/Runtime/Core/Public/ProfilingDebugging/TraceAuxiliary.h
+++ b/Engine/Source/Runtime/Core/Public/ProfilingDebugging/TraceAuxiliary.h
@@ -59,7 +59,7 @@ public:
 	* Callback whenever a trace snapshot is saved.
 	* Path is the file system path of the snapshot file.
 	*/
-	DECLARE_TS_MULTICAST_DELEGATE_OneParam(FOnSnapshotSaved, const FString& Path);
+	DECLARE_TS_MULTICAST_DELEGATE_TwoParams(FOnSnapshotSaved, FTraceAuxiliary::EConnectionType TraceType, const FString& TraceDestination);
 
 	struct FOptions
 	{
