Proper handing of live sessions in the recent traces list from the Insights Status Bar Widget.

diff --git a/Engine/Plugins/TraceUtilities/Source/EditorTraceUtilities/Private/InsightsStatusBar/SInsightsStatusBar.cpp b/Engine/Plugins/TraceUtilities/Source/EditorTraceUtilities/Private/InsightsStatusBar/SInsightsStatusBar.cpp
--- a/Engine/Plugins/TraceUtilities/Source/EditorTraceUtilities/Private/InsightsStatusBar/SInsightsStatusBar.cpp
+++ b/Engine/Plugins/TraceUtilities/Source/EditorTraceUtilities/Private/InsightsStatusBar/SInsightsStatusBar.cpp
@@ -239,11 +239,15 @@ void SInsightsStatusBarWidget::Construct(const FArguments& InArgs)
 	FTraceAuxiliary::OnTraceStopped.AddSP(this, &SInsightsStatusBarWidget::OnTraceStopped);
 	FTraceAuxiliary::OnSnapshotSaved.AddSP(this, &SInsightsStatusBarWidget::OnSnapshotSaved);
 
+	LiveSessionTracker = MakeShared<FLiveSessionTracker>();
+
 	InitCommandList();
 }
 
 TSharedRef<SWidget> SInsightsStatusBarWidget::MakeTraceMenu()
 {
+	LiveSessionTracker->StartQuery();
+
 	FMenuBuilder MenuBuilder(true, CommandList.ToSharedRef());
 
 	MenuBuilder.BeginSection("TraceData", LOCTEXT("TraceMenu_Section_Data", "Trace Data"));
@@ -471,7 +475,7 @@ void SInsightsStatusBarWidget::Traces_BuildMenu(FMenuBuilder& MenuBuilder)
 		MenuBuilder.AddMenuEntry(
 			FUIAction(FExecuteAction::CreateSP(this, &SInsightsStatusBarWidget::OpenTrace, Index),
 					  FCanExecuteAction()),
-			SNew(SRecentTracesListEntry, Traces[Index], TraceStorePath),
+			SNew(SRecentTracesListEntry, Traces[Index], TraceStorePath, LiveSessionTracker),
 			NAME_None,
 			FText::FromString(Traces[Index]->FilePath),
 			EUserInterfaceActionType::Button);
@@ -1011,6 +1015,22 @@ void SInsightsStatusBarWidget::OpenTrace(int32 Index)
 {
 	if (Index < Traces.Num())
 	{
+		if (FTraceAuxiliary::GetConnectionType() == FTraceAuxiliary::EConnectionType::Network)
+		{
+			if (LiveSessionTracker->HasData())
+			{
+				const FLiveSessionsMap& Sessions = LiveSessionTracker->GetLiveSessions();
+				FString FileName = FPaths::GetBaseFilename(Traces[Index]->FilePath);
+				const uint32* TraceId = Sessions.Find(FileName);
+				
+				if (TraceId)
+				{
+					FUnrealInsightsLauncher::Get()->OpenRemoteTrace(TEXT("localhost"), LiveSessionTracker->GetStorePort(), *TraceId);
+					return;
+				}
+			}
+		}
+
 		FUnrealInsightsLauncher::Get()->TryOpenTraceFromDestination(Traces[Index]->FilePath);
 	}
 }
diff --git a/Engine/Plugins/TraceUtilities/Source/EditorTraceUtilities/Private/InsightsStatusBar/SInsightsStatusBar.h b/Engine/Plugins/TraceUtilities/Source/EditorTraceUtilities/Private/InsightsStatusBar/SInsightsStatusBar.h
--- a/Engine/Plugins/TraceUtilities/Source/EditorTraceUtilities/Private/InsightsStatusBar/SInsightsStatusBar.h
+++ b/Engine/Plugins/TraceUtilities/Source/EditorTraceUtilities/Private/InsightsStatusBar/SInsightsStatusBar.h
@@ -8,6 +8,7 @@
 #include "Widgets/SWidget.h"
 #include "Widgets/SCompoundWidget.h"
 
+class FLiveSessionTracker;
 class FMenuBuilder;
 class FUICommandList;
 
@@ -155,6 +156,8 @@ private:
 	TArray<FChannelData> ChannelsInfo;
 	bool bShouldUpdateChannels = false;
 
+	TSharedPtr<FLiveSessionTracker> LiveSessionTracker;
+
 	TSharedPtr<FUICommandList> CommandList;
 
 	TArray<TSharedPtr<FTraceFileInfo>> Traces;
diff --git a/Engine/Plugins/TraceUtilities/Source/EditorTraceUtilities/Private/InsightsStatusBar/SRecentTracesList.cpp b/Engine/Plugins/TraceUtilities/Source/EditorTraceUtilities/Private/InsightsStatusBar/SRecentTracesList.cpp
--- a/Engine/Plugins/TraceUtilities/Source/EditorTraceUtilities/Private/InsightsStatusBar/SRecentTracesList.cpp
+++ b/Engine/Plugins/TraceUtilities/Source/EditorTraceUtilities/Private/InsightsStatusBar/SRecentTracesList.cpp
@@ -15,12 +15,13 @@
 
 #define LOCTEXT_NAMESPACE "RecentTracesList"
 
-void SRecentTracesListEntry::Construct(const FArguments& InArgs, TSharedPtr<FTraceFileInfo> InTrace, const FString& InStorePath)
+void SRecentTracesListEntry::Construct(const FArguments& InArgs, TSharedPtr<FTraceFileInfo> InTrace, const FString& InStorePath, TSharedPtr<FLiveSessionTracker> InLiveSessionTracker)
 {
 	TraceInfo = InTrace;
 	StorePath = InStorePath;
+	LiveSessionTracker = InLiveSessionTracker;
+	TraceName = FPaths::GetBaseFilename(TraceInfo->FilePath);
 
-	FString TraceName = FPaths::GetBaseFilename(TraceInfo->FilePath);
 	TSharedRef<FSlateFontMeasure> FontMeasureService = FSlateApplication::Get().GetRenderer()->GetFontMeasureService();
 	int32 LastCharacter = FontMeasureService->FindLastWholeCharacterIndexBeforeOffset(TraceName, FAppStyle::GetFontStyle("Menu.Label"), 180);
 	if (LastCharacter < TraceName.Len() - 1 && LastCharacter > 3)
@@ -66,7 +67,7 @@ void SRecentTracesListEntry::Construct(const FArguments& InArgs, TSharedPtr<FTra
 			+ SHorizontalBox::Slot()
 			.VAlign(EVerticalAlignment::VAlign_Center)
 			.HAlign(EHorizontalAlignment::HAlign_Left)
-			.Padding(0.0f, 2.0f, 0.0f, 0.0f)
+			.Padding(2.0f, 2.0f, 0.0f, 0.0f)
 			.AutoWidth()
 			[
 				SNew(STextBlock)
@@ -75,6 +76,19 @@ void SRecentTracesListEntry::Construct(const FArguments& InArgs, TSharedPtr<FTra
 				.ToolTipText(FText::FromString(TooltipText))
 				.ColorAndOpacity(FSlateColor::UseForeground())
 			]
+
+			+ SHorizontalBox::Slot()
+			.VAlign(EVerticalAlignment::VAlign_Center)
+			.HAlign(EHorizontalAlignment::HAlign_Left)
+			.Padding(3.0f, 1.0f, 0.0f, 0.0f)
+			.AutoWidth()
+			[
+				SNew(STextBlock)
+				.ToolTipText(LOCTEXT("LiveSessionTooltip", "This session is currently recording."))
+				.Text(LOCTEXT("LiveSessionLabel", "(LIVE)"))
+				.Visibility(this, &SRecentTracesListEntry::GetLiveLabelVisibility)
+				.ColorAndOpacity(FStyleColors::AccentRed)
+			]
 					
 			+ SHorizontalBox::Slot()
 			.VAlign(EVerticalAlignment::VAlign_Center)
@@ -104,4 +118,17 @@ FReply SRecentTracesListEntry::OpenContainingFolder()
 	return FReply::Handled();
 }
 
+EVisibility SRecentTracesListEntry::GetLiveLabelVisibility() const
+{
+	if (LiveSessionTracker.IsValid() && LiveSessionTracker->HasData())
+	{
+		if (LiveSessionTracker->GetLiveSessions().Contains(TraceName))
+		{
+			return EVisibility::Visible;
+		}
+	}
+
+	return EVisibility::Collapsed;
+}
+
 #undef LOCTEXT_NAMESPACE
diff --git a/Engine/Plugins/TraceUtilities/Source/EditorTraceUtilities/Private/InsightsStatusBar/SRecentTracesList.h b/Engine/Plugins/TraceUtilities/Source/EditorTraceUtilities/Private/InsightsStatusBar/SRecentTracesList.h
--- a/Engine/Plugins/TraceUtilities/Source/EditorTraceUtilities/Private/InsightsStatusBar/SRecentTracesList.h
+++ b/Engine/Plugins/TraceUtilities/Source/EditorTraceUtilities/Private/InsightsStatusBar/SRecentTracesList.h
@@ -6,6 +6,7 @@
 #include "Widgets/Views/SListView.h"
 
 struct FTraceFileInfo;
+class FLiveSessionTracker;
 
 /**
  *  A widget that shows an entry in the recent traces submenu.
@@ -15,12 +16,16 @@ class SRecentTracesListEntry : public SCompoundWidget
 	SLATE_BEGIN_ARGS(SRecentTracesListEntry) {}
 	SLATE_END_ARGS()
 
-	void Construct(const FArguments& InArgs, TSharedPtr<FTraceFileInfo> InTrace, const FString& InStorePath);
+	void Construct(const FArguments& InArgs, TSharedPtr<FTraceFileInfo> InTrace, const FString& InStorePath, TSharedPtr<FLiveSessionTracker> InLiveSessionTracker);
 
 private:
 	FReply OpenContainingFolder();
 
+	EVisibility GetLiveLabelVisibility() const;
+
 private:
 	FString StorePath;
 	TSharedPtr<FTraceFileInfo> TraceInfo;
+	FString TraceName;
+	TSharedPtr<FLiveSessionTracker> LiveSessionTracker;
 };
diff --git a/Engine/Plugins/TraceUtilities/Source/EditorTraceUtilities/Private/UnrealInsightsLauncher.cpp b/Engine/Plugins/TraceUtilities/Source/EditorTraceUtilities/Private/UnrealInsightsLauncher.cpp
--- a/Engine/Plugins/TraceUtilities/Source/EditorTraceUtilities/Private/UnrealInsightsLauncher.cpp
+++ b/Engine/Plugins/TraceUtilities/Source/EditorTraceUtilities/Private/UnrealInsightsLauncher.cpp
@@ -3,6 +3,7 @@
 #include "UnrealInsightsLauncher.h"
 #include "EditorTraceUtilities.h"
 
+#include "Async/TaskGraphInterfaces.h"
 #include "Styling/AppStyle.h"
 #include "IUATHelperModule.h"
 #include "Logging/LogMacros.h"
@@ -14,7 +15,6 @@
 
 #define LOCTEXT_NAMESPACE "FUnrealInsightsLauncher"
 
-
 TSharedPtr<FUnrealInsightsLauncher> FUnrealInsightsLauncher::Instance = nullptr;
 
 // We use this Task to launch notifications from the Game Thread because on Mac the app is closed if the notification API is called from another thread.
@@ -38,6 +38,60 @@ private:
 	FText Message;
 };
 
+class FLiveSessionQueryTask
+{
+public:
+	FLiveSessionQueryTask(TSharedPtr<FLiveSessionTaskData> InOutData)
+		: OutData(InOutData)
+	{}
+
+	FORCEINLINE TStatId GetStatId() const { RETURN_QUICK_DECLARE_CYCLE_STAT(FLiveSessionQueryTask, STATGROUP_TaskGraphTasks); }
+	ENamedThreads::Type GetDesiredThread() { return ENamedThreads::Type::AnyThread; }
+	static ESubsequentsMode::Type GetSubsequentsMode() { return ESubsequentsMode::TrackSubsequents; }
+
+	void DoTask(ENamedThreads::Type CurrentThread, const FGraphEventRef& MyCompletionGraphEvent)
+	{
+		using namespace UE::Trace;
+		FStoreClient* StoreClient = FStoreClient::Connect(TEXT("localhost"));
+
+		OutData->TaskLiveSessionData.Empty();
+
+		if (!StoreClient)
+		{
+			return;
+		}
+
+		uint32 SessionCount = StoreClient->GetSessionCount();
+		if (!SessionCount)
+		{
+			return;
+		}
+		
+		OutData->StorePort = StoreClient->GetStorePort();
+
+		for (uint32 Index = 0; Index < SessionCount; ++Index)
+		{
+			const FStoreClient::FSessionInfo* SessionInfo = StoreClient->GetSessionInfo(Index);
+
+			if (!SessionInfo)
+			{
+				continue;
+			}
+
+			uint32 TraceId = SessionInfo->GetTraceId();
+			const FStoreClient::FTraceInfo* Info = StoreClient->GetTraceInfoById(TraceId);
+
+			if (Info)
+			{
+				OutData->TaskLiveSessionData.Add(FString(Info->GetName()), TraceId);
+			}
+		}
+	}
+
+	TSharedPtr<FLiveSessionTaskData> OutData;
+};
+
+
 FUnrealInsightsLauncher::FUnrealInsightsLauncher()
 	: LogListingName(TEXT("UnrealInsights"))
 {
@@ -216,4 +270,51 @@ bool FUnrealInsightsLauncher::OpenActiveTraceFromStore(const FString& TraceHostA
 	return true;
 }
 
+FLiveSessionTracker::FLiveSessionTracker()
+{
+	TaskLiveSessionData = MakeShared<FLiveSessionTaskData>();
+}
+
+void FLiveSessionTracker::Update()
+{
+	if (bIsQueryInProgress && Event.IsValid() && Event->IsComplete())
+	{
+		bIsQueryInProgress = false;
+		LiveSessionMap = TaskLiveSessionData->TaskLiveSessionData;
+		StorePort = TaskLiveSessionData->StorePort;
+
+		if (!LiveSessionMap.IsEmpty())
+		{
+			bHasData = true;
+		}
+	}
+}
+
+void FLiveSessionTracker::StartQuery()
+{
+	if (!bIsQueryInProgress)
+	{
+		Event = TGraphTask<FLiveSessionQueryTask>::CreateTask().ConstructAndDispatchWhenReady(TaskLiveSessionData);
+		bIsQueryInProgress = true;
+	}
+}
+
+const FLiveSessionsMap& FLiveSessionTracker::GetLiveSessions()
+{
+	Update();
+	return LiveSessionMap;
+}
+
+bool FLiveSessionTracker::HasData()
+{
+	Update();
+	return bHasData;
+}
+
+uint32 FLiveSessionTracker::GetStorePort()
+{
+	Update();
+	return StorePort;
+}
+
 #undef LOCTEXT_NAMESPACE
diff --git a/Engine/Plugins/TraceUtilities/Source/EditorTraceUtilities/Private/UnrealInsightsLauncher.h b/Engine/Plugins/TraceUtilities/Source/EditorTraceUtilities/Private/UnrealInsightsLauncher.h
--- a/Engine/Plugins/TraceUtilities/Source/EditorTraceUtilities/Private/UnrealInsightsLauncher.h
+++ b/Engine/Plugins/TraceUtilities/Source/EditorTraceUtilities/Private/UnrealInsightsLauncher.h
@@ -2,6 +2,7 @@
 
 #pragma once
 
+#include "Async/TaskGraphFwd.h"
 #include "CoreMinimal.h"
 
 class FUnrealInsightsLauncher : public TSharedFromThis<FUnrealInsightsLauncher>
@@ -83,3 +84,35 @@ private:
 	FName LogListingName;
 };
 
+typedef TMap<FString, uint32> FLiveSessionsMap;
+
+struct FLiveSessionTaskData
+{
+	FLiveSessionsMap TaskLiveSessionData;
+	uint32 StorePort;
+};
+
+class FLiveSessionTracker
+{
+public:
+	FLiveSessionTracker();
+	~FLiveSessionTracker() {}
+
+	void Update();
+	void StartQuery();
+
+	bool HasData();
+	const FLiveSessionsMap& GetLiveSessions();
+	uint32 GetStorePort();
+
+private:
+	bool bHasData = false;
+	bool bIsQueryInProgress = false;
+	uint32 StorePort;
+
+	FLiveSessionsMap LiveSessionMap;
+	TSharedPtr<FLiveSessionTaskData> TaskLiveSessionData;
+
+	FGraphEventRef Event;
+};
+
diff --git a/Engine/Source/Developer/TraceInsights/Private/Insights/Widgets/SStartPageWindow.cpp b/Engine/Source/Developer/TraceInsights/Private/Insights/Widgets/SStartPageWindow.cpp
--- a/Engine/Source/Developer/TraceInsights/Private/Insights/Widgets/SStartPageWindow.cpp
+++ b/Engine/Source/Developer/TraceInsights/Private/Insights/Widgets/SStartPageWindow.cpp
@@ -190,6 +190,7 @@ public:
 					SNew(STextBlock)
 					.Text(this, &STraceListRow::GetTraceStatus)
 					.ToolTip(STraceListRow::GetTraceTooltip())
+					.ColorAndOpacity(FStyleColors::AccentRed)
 				];
 		}
 		else
