When opening the profiling directory from the Insights Status Bar Widget, the folder will be created if it does not exist.

diff --git a/Engine/Plugins/TraceUtilities/Source/EditorTraceUtilities/Private/InsightsStatusBar/SInsightsStatusBar.cpp b/Engine/Plugins/TraceUtilities/Source/EditorTraceUtilities/Private/InsightsStatusBar/SInsightsStatusBar.cpp
--- a/Engine/Plugins/TraceUtilities/Source/EditorTraceUtilities/Private/InsightsStatusBar/SInsightsStatusBar.cpp
+++ b/Engine/Plugins/TraceUtilities/Source/EditorTraceUtilities/Private/InsightsStatusBar/SInsightsStatusBar.cpp
@@ -587,6 +587,12 @@ void SInsightsStatusBarWidget::OpenProfilingDirectory_OnClicked()
 void SInsightsStatusBarWidget::OpenProfilingDirectory()
 {
 	FString FullPath(FPaths::ConvertRelativePathToFull(FPaths::ProfilingDir()));
+
+	if (!IFileManager::Get().DirectoryExists(*FullPath))
+	{
+		IFileManager::Get().MakeDirectory(*FullPath);
+	}
+
 	FPlatformProcess::ExploreFolder(*FullPath);
 }
 
