Fix Insights trace not opening correctly if it's path contains spaces.
Update Unreal Insights version.

diff --git a/Engine/Plugins/TraceUtilities/Source/EditorTraceUtilities/Private/UnrealInsightsLauncher.cpp b/Engine/Plugins/TraceUtilities/Source/EditorTraceUtilities/Private/UnrealInsightsLauncher.cpp
--- a/Engine/Plugins/TraceUtilities/Source/EditorTraceUtilities/Private/UnrealInsightsLauncher.cpp
+++ b/Engine/Plugins/TraceUtilities/Source/EditorTraceUtilities/Private/UnrealInsightsLauncher.cpp
@@ -211,6 +211,18 @@ bool FUnrealInsightsLauncher::TryOpenTraceFromDestination(const FString& Destina
 
 bool FUnrealInsightsLauncher::OpenTraceFile(const FString& FilePath)
 {
+	if (!FilePath.StartsWith(TEXT("\"")) && !FilePath.EndsWith(TEXT("\"")))
+	{
+		TStringBuilder<1024> StringBuilder;
+
+		StringBuilder.AppendChar(TEXT('"'));
+		StringBuilder.Append(FilePath);
+		StringBuilder.AppendChar(TEXT('"'));
+
+		StartUnrealInsights(GetInsightsApplicationPath(), StringBuilder.ToString());
+		return true;
+	}
+
 	StartUnrealInsights(GetInsightsApplicationPath(), FilePath);
 	return true;
 }
diff --git a/Engine/Source/Developer/TraceInsights/Public/Insights/Version.h b/Engine/Source/Developer/TraceInsights/Public/Insights/Version.h
--- a/Engine/Source/Developer/TraceInsights/Public/Insights/Version.h
+++ b/Engine/Source/Developer/TraceInsights/Public/Insights/Version.h
@@ -5,10 +5,10 @@
 // See https://semver.org/
 
 #define UNREAL_INSIGHTS_VERSION_MAJOR 1
-#define UNREAL_INSIGHTS_VERSION_MINOR 2
+#define UNREAL_INSIGHTS_VERSION_MINOR 3
 #define UNREAL_INSIGHTS_VERSION_PATCH 0
 
-#define UNREAL_INSIGHTS_VERSION_STRING "1.02"
+#define UNREAL_INSIGHTS_VERSION_STRING "1.03"
 
 /** Extra identifier string (like: "alpha", "beta.1", "beta.2", "test", "featureX", etc.). */
 //#define UNREAL_INSIGHTS_VERSION_ID_STRING "beta.1"
