[Insights] Fixed public API comments in TraceUtilLibrary.h.

diff --git a/Engine/Plugins/TraceUtilities/Source/TraceUtilities/Public/TraceUtilLibrary.h b/Engine/Plugins/TraceUtilities/Source/TraceUtilities/Public/TraceUtilLibrary.h
--- a/Engine/Plugins/TraceUtilities/Source/TraceUtilities/Public/TraceUtilLibrary.h
+++ b/Engine/Plugins/TraceUtilities/Source/TraceUtilities/Public/TraceUtilLibrary.h
@@ -1,5 +1,6 @@
 // Copyright Epic Games, Inc. All Rights Reserved.
-#pragma once 
+
+#pragma once
 
 #include "Kismet/BlueprintFunctionLibrary.h"
 #include "TraceUtilLibrary.generated.h"
@@ -7,52 +8,54 @@
 UCLASS()
 class TRACEUTILITIES_API UTraceUtilLibrary : public UBlueprintFunctionLibrary
 {
-    GENERATED_BODY()
-    
-    public:
-    
-    UFUNCTION(BlueprintCallable, Category = "Perf | Insights Trace")
-    static bool StartTraceToFile(const FString& FileName, const TArray<FString>& Channels);
-    
-    UFUNCTION(BlueprintCallable, Category = "Perf | Insights Trace")
-    static bool StartTraceSendTo(const FString& Target, const TArray<FString>& Channels);
-    
-    UFUNCTION(BlueprintCallable, Category = "Perf | Insights Trace")
-    static bool StopTracing();
-    
-    UFUNCTION(BlueprintCallable, Category = "Perf | Insights Trace")
-    static bool PauseTracing();
-    
-    UFUNCTION(BlueprintCallable, Category = "Perf | Insights Trace")
-    static bool ResumeTracing();
-    
-    UFUNCTION(BlueprintCallable, Category = "Perf | Insights Trace")
-    static bool IsTracing();
-    
-    UFUNCTION(BlueprintCallable, Category = "Perf | Insights Trace")
-    static bool ToggleChannel(const FString& ChannelName, bool enabled);
-    
-    UFUNCTION(BlueprintCallable, Category = "Perf | Insights Trace")
-    static bool IsChannelEnabled(const FString& ChannelName);
-    
-    UFUNCTION(BlueprintCallable, Category = "Perf | Insights Trace")
-    static TArray<FString> GetEnabledChannels();
-    
-    UFUNCTION(BlueprintCallable, Category = "Perf | Insights Trace")
-    static TArray<FString> GetAllChannels();
-    
-    UFUNCTION(BlueprintCallable, Category = "Perf | Insights Trace")
-    static void TraceBookmark(const FString& Name);
-
-    /**
-     * Traces a "RegionStart:Name" bookmark, where Name is defined through the Name parameter.
-     */
-    UFUNCTION(BlueprintCallable, Category = "Perf | Insights Trace")
-    static void TraceMarkRegionStart(const FString& Name);
- 
+	GENERATED_BODY()
+
+public:
+	UFUNCTION(BlueprintCallable, Category = "Perf | Insights Trace")
+	static bool StartTraceToFile(const FString& FileName, const TArray<FString>& Channels);
+
+	UFUNCTION(BlueprintCallable, Category = "Perf | Insights Trace")
+	static bool StartTraceSendTo(const FString& Target, const TArray<FString>& Channels);
+
+	UFUNCTION(BlueprintCallable, Category = "Perf | Insights Trace")
+	static bool StopTracing();
+
+	UFUNCTION(BlueprintCallable, Category = "Perf | Insights Trace")
+	static bool PauseTracing();
+
+	UFUNCTION(BlueprintCallable, Category = "Perf | Insights Trace")
+	static bool ResumeTracing();
+
+	UFUNCTION(BlueprintCallable, Category = "Perf | Insights Trace")
+	static bool IsTracing();
+
+	UFUNCTION(BlueprintCallable, Category = "Perf | Insights Trace")
+	static bool ToggleChannel(const FString& ChannelName, bool enabled);
+
+	UFUNCTION(BlueprintCallable, Category = "Perf | Insights Trace")
+	static bool IsChannelEnabled(const FString& ChannelName);
+
+	UFUNCTION(BlueprintCallable, Category = "Perf | Insights Trace")
+	static TArray<FString> GetEnabledChannels();
+
+	UFUNCTION(BlueprintCallable, Category = "Perf | Insights Trace")
+	static TArray<FString> GetAllChannels();
+
+	/**
+	 * Traces a bookmark with specified name.
+	 */
+	UFUNCTION(BlueprintCallable, Category = "Perf | Insights Trace")
+	static void TraceBookmark(const FString& Name);
+
+	/**
+	 * Traces a begin event for a region with specified name.
+	 */
+	UFUNCTION(BlueprintCallable, Category = "Perf | Insights Trace")
+	static void TraceMarkRegionStart(const FString& Name);
+
 	/**
-	 * Traces a "RegionEnd:Name" bookmark, where Name is defined through the Name parameter.
-	 */   
-    UFUNCTION(BlueprintCallable, Category = "Perf | Insights Trace")
-    static void TraceMarkRegionEnd(const FString& Name);
-};
\ No newline at end of file
+	 * Traces an end event for a region with specified name.
+	 */
+	UFUNCTION(BlueprintCallable, Category = "Perf | Insights Trace")
+	static void TraceMarkRegionEnd(const FString& Name);
+};
