Use new trace functionality to correctly identify live session from own process

Previously executing "Open current session" in the Insights launcher used the last activated live session, which might be incorrect if the user has multiple live session at the same time. This change uses new functionality in TraceLog which adds session and trace guid fields as metadata and queries the server for the current ongoing session.

diff --git a/Engine/Plugins/TraceUtilities/Source/EditorTraceUtilities/Private/UnrealInsightsLauncher.cpp b/Engine/Plugins/TraceUtilities/Source/EditorTraceUtilities/Private/UnrealInsightsLauncher.cpp
--- a/Engine/Plugins/TraceUtilities/Source/EditorTraceUtilities/Private/UnrealInsightsLauncher.cpp
+++ b/Engine/Plugins/TraceUtilities/Source/EditorTraceUtilities/Private/UnrealInsightsLauncher.cpp
@@ -244,35 +244,55 @@ bool FUnrealInsightsLauncher::OpenRemoteTrace(const FString& TraceHostAddress, c
 bool FUnrealInsightsLauncher::OpenActiveTraceFromStore(const FString& TraceHostAddress)
 {
 	UE::Trace::FStoreClient* StoreClient = UE::Trace::FStoreClient::Connect(*TraceHostAddress);
-	
+
 	if (!StoreClient)
 	{
-		const FText	MessageBoxTextFmt = LOCTEXT("StoreClientConnectFailed_TextFmt", "Could not connect to StoreClient at {0}");
-     	const FText MessageBoxText = FText::Format(MessageBoxTextFmt, FText::FromString(TraceHostAddress));
-     	LogMessage(MessageBoxText);
+		const FText MessageBoxTextFmt = LOCTEXT("StoreClientConnectFailed_TextFmt",
+												"Could not connect to StoreClient at {0}");
+		const FText MessageBoxText = FText::Format(MessageBoxTextFmt, FText::FromString(TraceHostAddress));
+		LogMessage(MessageBoxText);
 		//No active connection to store client
 		return false;
 	}
-	 
+
 	int SessionCount = StoreClient->GetSessionCount();
 	if (!SessionCount)
 	{
-		const FText	MessageBoxTextFmt = LOCTEXT("StoreClientNoSession_TextFmt", "No active session found in StoreClient at {0}");
-     	const FText MessageBoxText = FText::Format(MessageBoxTextFmt, FText::FromString(TraceHostAddress));
-     	LogMessage(MessageBoxText);
+		const FText MessageBoxTextFmt = LOCTEXT("StoreClientNoSession_TextFmt",
+												"No active session found in StoreClient at {0}");
+		const FText MessageBoxText = FText::Format(MessageBoxTextFmt, FText::FromString(TraceHostAddress));
+		LogMessage(MessageBoxText);
 		return false;
 	}
 	// Get first active session
-	const UE::Trace::FStoreClient::FSessionInfo* SessionInfo = StoreClient->GetSessionInfo(0);
+	FGuid SessionGuid, TraceGuid;
+	if (!FTraceAuxiliary::IsConnected(SessionGuid, TraceGuid))
+	{
+		return false;
+	}
 
+	// TraceGuid is enough to identify the session
+	const UE::Trace::FStoreClient::FSessionInfo* SessionInfo = StoreClient->GetSessionInfoByGuid(TraceGuid);
+
+	// Fallback to old method of using last active session, in case running with old server
+	if (!SessionInfo)
+	{
+		UE_LOG(LogTraceUtilities, Warning, TEXT("Unable to find session using guid. Possibly running with older server version. Falling back to last active session."));
+		SessionInfo = StoreClient->GetSessionInfo(0);
+	}
+	
 	if (!SessionInfo)
 	{
 		// Failed to retrieve SessionInfo for active Session with Index 0
+		const FText MessageBoxTextFmt = LOCTEXT("StoreClientFailedToRetrieve_TextFmt",
+												"Failed to retrieve active session in StoreClient at {0}");
+		const FText MessageBoxText = FText::Format(MessageBoxTextFmt, FText::FromString(TraceHostAddress));
+		LogMessage(MessageBoxText);
 		return false;
 	}
-	
+
 	uint32 TraceId = SessionInfo->GetTraceId();
-	
+
 	if (!OpenRemoteTrace(TraceHostAddress, (uint16)StoreClient->GetStorePort(), TraceId))
 	{
 		//Failed to open Trace File %s
diff --git a/Engine/Source/Developer/TraceAnalysis/Private/Store/CborPayload.h b/Engine/Source/Developer/TraceAnalysis/Private/Store/CborPayload.h
--- a/Engine/Source/Developer/TraceAnalysis/Private/Store/CborPayload.h
+++ b/Engine/Source/Developer/TraceAnalysis/Private/Store/CborPayload.h
@@ -74,6 +74,8 @@ public:
 	template <int N> void		AddInteger(const char (&Name)[N], int64 Value);
 	template <int N> void		AddString(const char (&Name)[N], const char* Value, int32 Length=-1);
 	template <int N> void		AddStringArray(const char(&Name)[N], const TArray<FString>& Values);
+	template <int N, typename T>
+					void		AddArray(const char(&Name)[N], const TArrayView<T>& Values);
 	FPayload					Done();
 
 private:
@@ -133,7 +135,20 @@ inline void TPayloadBuilder<Size>::AddStringArray(const char(&Name)[N], const TA
 		CborWriter.WriteValue(Elem);
 	}
 }
-	
+
+////////////////////////////////////////////////////////////////////////////////
+template <int Size>
+template <int N, typename T>
+void TPayloadBuilder<Size>::AddArray(const char(& Name)[N], const TArrayView<T>& Values)
+{
+	CborWriter.WriteValue(Name, N - 1);
+	CborWriter.WriteContainerStart(ECborCode::Array, Values.Num());
+	for (const T& Elem : Values)
+	{
+		CborWriter.WriteValue(Elem);
+	}
+}
+
 ////////////////////////////////////////////////////////////////////////////////
 template <int Size>
 inline FPayload TPayloadBuilder<Size>::Done()
diff --git a/Engine/Source/Developer/TraceAnalysis/Private/Store/StoreClient.cpp b/Engine/Source/Developer/TraceAnalysis/Private/Store/StoreClient.cpp
--- a/Engine/Source/Developer/TraceAnalysis/Private/Store/StoreClient.cpp
+++ b/Engine/Source/Developer/TraceAnalysis/Private/Store/StoreClient.cpp
@@ -3,7 +3,6 @@
 #include "Trace/StoreClient.h"
 #include "Algo/Transform.h"
 #include "Asio/Asio.h"
-#include "AsioStore.h"
 #include "CborPayload.h"
 #include "Templates/UnrealTemplate.h"
 
@@ -116,6 +115,7 @@ public:
 	bool					GetSessionInfo(uint32 Index);
 	bool					GetSessionInfoById(uint32 Id);
 	bool					GetSessionInfoByTraceId(uint32 TraceId);
+	bool					GetSessionInfoByTraceGuid(const FGuid& TraceGuid);
 							bool SetStoreDirectories(const TCHAR* StoreDir, const TArray<FString>& AddWatchDirs,
 													const TArray<FString>& RemoveWatchDirs);
 
@@ -380,6 +380,15 @@ bool FStoreCborClient::GetSessionInfoByTraceId(uint32 TraceId)
 	return Communicate(Payload);
 }
 
+////////////////////////////////////////////////////////////////////////////////
+bool FStoreCborClient::GetSessionInfoByTraceGuid(const FGuid& TraceGuid)
+{
+	TPayloadBuilder<> Builder("v1/session/info");
+	Builder.AddString("trace_guid", TCHAR_TO_ANSI(*TraceGuid.ToString()));
+	FPayload Payload = Builder.Done();
+	return Communicate(Payload);
+}
+
 ////////////////////////////////////////////////////////////////////////////////
 bool FStoreCborClient::SetStoreDirectories(const TCHAR* StoreDir, const TArray<FString>& AddWatchDirs, const TArray<FString>& RemoveWatchDir)
 {
@@ -679,5 +688,18 @@ const FStoreClient::FSessionInfo* FStoreClient::GetSessionInfoByTraceId(uint32 T
 	return (FSessionInfo*)(&Response);
 }
 
+////////////////////////////////////////////////////////////////////////////////
+const FStoreClient::FSessionInfo* FStoreClient::GetSessionInfoByGuid(const FGuid& TraceGuid)
+{
+	auto* Self = (FStoreCborClient*)this;
+	if (!Self->GetSessionInfoByTraceGuid(TraceGuid))
+	{
+		return nullptr;
+	}
+
+	const FResponse& Response = Self->GetResponse();
+	return (FSessionInfo*)(&Response);
+}
+	
 } // namespace Trace
 } // namespace UE
diff --git a/Engine/Source/Developer/TraceAnalysis/Public/Trace/StoreClient.h b/Engine/Source/Developer/TraceAnalysis/Public/Trace/StoreClient.h
--- a/Engine/Source/Developer/TraceAnalysis/Public/Trace/StoreClient.h
+++ b/Engine/Source/Developer/TraceAnalysis/Public/Trace/StoreClient.h
@@ -72,6 +72,7 @@ public:
 	const FSessionInfo* GetSessionInfo(uint32 Index) const;
 	const FSessionInfo* GetSessionInfoById(uint32 Id) const;
 	const FSessionInfo* GetSessionInfoByTraceId(uint32 TraceId) const;
+	const FSessionInfo* GetSessionInfoByGuid(const FGuid& TraceGuid);
 #if 0
 	template <typename Lambda> uint32	GetSessionInfos(uint32 StartIndex, uint32 Count, Lambda&& Callback) const;
 #endif // 0
