Enable functionality in Insights toolbar and session browser for all platforms.

Some options in the toolbar and the session browser were not displayed on Linux and Mac since they didn't have UnrealTraceServer. With this feature enabled we can display all options on all platforms.

diff --git a/Engine/Plugins/TraceUtilities/Source/EditorTraceUtilities/Private/InsightsStatusBar/SInsightsStatusBar.cpp b/Engine/Plugins/TraceUtilities/Source/EditorTraceUtilities/Private/InsightsStatusBar/SInsightsStatusBar.cpp
--- a/Engine/Plugins/TraceUtilities/Source/EditorTraceUtilities/Private/InsightsStatusBar/SInsightsStatusBar.cpp
+++ b/Engine/Plugins/TraceUtilities/Source/EditorTraceUtilities/Private/InsightsStatusBar/SInsightsStatusBar.cpp
@@ -343,7 +343,6 @@ TSharedRef<SWidget> SInsightsStatusBarWidget::MakeTraceMenu()
 
 	MenuBuilder.BeginSection("Options", LOCTEXT("TraceMenu_Section_Options", "Options"));
 	{
-#if PLATFORM_WINDOWS
 		MenuBuilder.AddMenuEntry(
 			LOCTEXT("OpenLiveSesssionOnTraceStart", "Open Live Session on Trace Start"),
 			LOCTEXT("OpenLiveSesssionOnTraceStartDesc", "When set, the live session will be automatically opened in Unreal Insights when tracing is started.\nThis option will only apply when tracing to the trace store."),
@@ -365,7 +364,6 @@ TSharedRef<SWidget> SInsightsStatusBarWidget::MakeTraceMenu()
 			NAME_None,
 			EUserInterfaceActionType::ToggleButton
 		);
-#endif
 
 		MenuBuilder.AddMenuEntry(
 			LOCTEXT("ShowInExplorerAfterTrace", "Show in Explorer after Trace"),
@@ -413,7 +411,6 @@ TSharedRef<SWidget> SInsightsStatusBarWidget::MakeTraceMenu()
 			EUserInterfaceActionType::Button
 		);
 
-#if PLATFORM_WINDOWS
 		MenuBuilder.AddMenuEntry(
 			LOCTEXT("OpenLiveSessionLabel", "Open Live Session"),
 			LOCTEXT("OpenLiveSessionTooltip", "Opening the live session is possible only while tracing to the trace store."),
@@ -423,7 +420,6 @@ TSharedRef<SWidget> SInsightsStatusBarWidget::MakeTraceMenu()
 			NAME_None,
 			EUserInterfaceActionType::Button
 		);
-#endif
 
 		MenuBuilder.AddSubMenu
 		(
@@ -863,7 +859,6 @@ void SInsightsStatusBarWidget::CacheTraceStorePath()
 {
 	if (TraceStorePath.IsEmpty())
 	{
-#if PLATFORM_WINDOWS
 		UE::Trace::FStoreClient* StoreClient = UE::Trace::FStoreClient::Connect(TEXT("localhost"));
 
 		if (!StoreClient)
@@ -880,9 +875,6 @@ void SInsightsStatusBarWidget::CacheTraceStorePath()
 		}
 
 		TraceStorePath = FString(Status->GetStoreDir());
-#else
-		TraceStorePath = FPaths::Combine(FPaths::EngineDir(), TEXT("/Programs/UnrealInsights/Saved/TraceSessions/"));
-#endif
 	}
 }
 
diff --git a/Engine/Source/Developer/TraceInsights/Private/Insights/Widgets/SStartPageWindow.cpp b/Engine/Source/Developer/TraceInsights/Private/Insights/Widgets/SStartPageWindow.cpp
--- a/Engine/Source/Developer/TraceInsights/Private/Insights/Widgets/SStartPageWindow.cpp
+++ b/Engine/Source/Developer/TraceInsights/Private/Insights/Widgets/SStartPageWindow.cpp
@@ -57,9 +57,6 @@
 
 #define LOCTEXT_NAMESPACE "STraceStoreWindow"
 
-// For now, we only support store settings on platforms that run Unreal Trace Server.
-#define STORE_SUPPORTS_SETTINGS PLATFORM_WINDOWS
-
 // Should match GDefaultChannels (from Runtime\Core\Private\ProfilingDebugging\TraceAuxiliary.cpp).
 // We cannot use GDefaultChannels directly as UE_TRACE_ENABLED may be off.
 static const TCHAR* GInsightsDefaultChannelPreset = TEXT("cpu,gpu,frame,log,bookmark,screenshot,region");
@@ -1642,7 +1639,6 @@ TSharedRef<SWidget> STraceStoreWindow::ConstructTraceStoreDirectoryPanel()
 		[
 			SNew(SHorizontalBox)
 
-#if STORE_SUPPORTS_SETTINGS
 			+ SHorizontalBox::Slot()
 			.AutoWidth()
 			.Padding(0.0f, 0.0f, 0.0f, 0.0f)
@@ -1658,7 +1654,6 @@ TSharedRef<SWidget> STraceStoreWindow::ConstructTraceStoreDirectoryPanel()
 					.ColorAndOpacity(FSlateColor::UseForeground())
 				]
 			]
-#endif
 
 			+ SHorizontalBox::Slot()
 			.AutoWidth()
@@ -1671,7 +1666,6 @@ TSharedRef<SWidget> STraceStoreWindow::ConstructTraceStoreDirectoryPanel()
 				.Text_Raw(this, &STraceStoreWindow::GetConnectionStatusTooltip))
 			]
 
-#if STORE_SUPPORTS_SETTINGS
 			+ SHorizontalBox::Slot()
 			.AutoWidth()
 			.Padding(4.0f, 0.0f, 0.0f, 0.0f)
@@ -1690,7 +1684,6 @@ TSharedRef<SWidget> STraceStoreWindow::ConstructTraceStoreDirectoryPanel()
 				.IsReadOnly(true)
 				.BackgroundColor(FSlateColor(EStyleColor::Background))
 			]
-#endif
 
 			+ SHorizontalBox::Slot()
 			.AutoWidth()
@@ -1730,7 +1723,6 @@ TSharedRef<SWidget> STraceStoreWindow::ConstructTraceStoreDirectoryPanel()
 			]
 		]
 
-#if STORE_SUPPORTS_SETTINGS
 		+ SVerticalBox::Slot()
 		.AutoHeight()
 		.MaxHeight(400)
@@ -1816,7 +1808,6 @@ TSharedRef<SWidget> STraceStoreWindow::ConstructTraceStoreDirectoryPanel()
 				]
 			]
 		]
-#endif
 		;
 
 	return Widget;
