Fixed deadlock when stopping a trace with the "Open Trace in Insights" option enabled

diff --git a/Engine/Plugins/TraceUtilities/Source/EditorTraceUtilities/Private/InsightsStatusBar/SInsightsStatusBar.cpp b/Engine/Plugins/TraceUtilities/Source/EditorTraceUtilities/Private/InsightsStatusBar/SInsightsStatusBar.cpp
--- a/Engine/Plugins/TraceUtilities/Source/EditorTraceUtilities/Private/InsightsStatusBar/SInsightsStatusBar.cpp
+++ b/Engine/Plugins/TraceUtilities/Source/EditorTraceUtilities/Private/InsightsStatusBar/SInsightsStatusBar.cpp
@@ -567,12 +567,12 @@ void SInsightsStatusBarWidget::LaunchUnrealInsights_OnClicked()
 
 void SInsightsStatusBarWidget::OpenLiveSession_OnClicked()
 {
-	OpenLiveSession();
+	OpenLiveSession(FTraceAuxiliary::GetTraceDestinationString());
 }
 
-void SInsightsStatusBarWidget::OpenLiveSession()
+void SInsightsStatusBarWidget::OpenLiveSession(const FString& InTraceDestination)
 {
-	FUnrealInsightsLauncher::Get()->TryOpenTraceFromDestination(FTraceAuxiliary::GetTraceDestinationString());
+	FUnrealInsightsLauncher::Get()->TryOpenTraceFromDestination(InTraceDestination);
 }
 
 void SInsightsStatusBarWidget::OpenProfilingDirectory_OnClicked()
@@ -813,7 +813,7 @@ void SInsightsStatusBarWidget::OnTraceStopped(FTraceAuxiliary::EConnectionType I
 {
 	if (GetBooleanSettingValue(OpenInsightsAfterTraceSettingName))
 	{
-		OpenLiveSession();
+		OpenLiveSession(InTraceDestination);
 	}
 	if (GetBooleanSettingValue(ShowInExplorerAfterTraceSettingName))
 	{
diff --git a/Engine/Plugins/TraceUtilities/Source/EditorTraceUtilities/Private/InsightsStatusBar/SInsightsStatusBar.h b/Engine/Plugins/TraceUtilities/Source/EditorTraceUtilities/Private/InsightsStatusBar/SInsightsStatusBar.h
--- a/Engine/Plugins/TraceUtilities/Source/EditorTraceUtilities/Private/InsightsStatusBar/SInsightsStatusBar.h
+++ b/Engine/Plugins/TraceUtilities/Source/EditorTraceUtilities/Private/InsightsStatusBar/SInsightsStatusBar.h
@@ -71,7 +71,7 @@ private:
 	void LaunchUnrealInsights_OnClicked();
 	
 	void OpenLiveSession_OnClicked();
-	void OpenLiveSession();
+	void OpenLiveSession(const FString& InTraceDestination);
 
 	void OpenProfilingDirectory_OnClicked();
 	void OpenProfilingDirectory();
