Automated C34850469 "Automate Test Case: Unreal Insights - Session Browser - Compiling and Opening"

diff --git a/Engine/Plugins/TraceUtilities/Source/EditorTraceUtilities/Private/UnrealInsightsLauncher.cpp b/Engine/Plugins/TraceUtilities/Source/EditorTraceUtilities/Private/UnrealInsightsLauncher.cpp
--- a/Engine/Plugins/TraceUtilities/Source/EditorTraceUtilities/Private/UnrealInsightsLauncher.cpp
+++ b/Engine/Plugins/TraceUtilities/Source/EditorTraceUtilities/Private/UnrealInsightsLauncher.cpp
@@ -128,9 +128,9 @@ void FUnrealInsightsLauncher::StartUnrealInsights(const FString& Path, const FSt
 
 	void* PipeWriteChild = nullptr;
 	void* PipeReadChild = nullptr;
-	FProcHandle Handle = FPlatformProcess::CreateProc(*Path, *Parameters, bLaunchDetached, bLaunchHidden, bLaunchReallyHidden, &ProcessID, PriorityModifier, OptionalWorkingDirectory, PipeWriteChild, PipeReadChild);
+	UnrealInsightsHandle = FPlatformProcess::CreateProc(*Path, *Parameters, bLaunchDetached, bLaunchHidden, bLaunchReallyHidden, &ProcessID, PriorityModifier, OptionalWorkingDirectory, PipeWriteChild, PipeReadChild);
 
-	if (Handle.IsValid())
+	if (UnrealInsightsHandle.IsValid())
 	{
 		UE_LOG(LogTraceUtilities, Log, TEXT("Launched Unreal Insights executable: %s %s"), *Path, *Parameters);
 	}
@@ -142,6 +142,17 @@ void FUnrealInsightsLauncher::StartUnrealInsights(const FString& Path, const FSt
 	}
 }
 
+void FUnrealInsightsLauncher::CloseUnrealInsights()
+{
+	if (UnrealInsightsHandle.IsValid())
+	{
+		FPlatformProcess::TerminateProc(UnrealInsightsHandle);
+	}
+	else
+	{
+		UE_LOG(LogTraceUtilities, Log, TEXT("Could not find the Unreal Insights process handler"));
+	}
+}
 
 void FUnrealInsightsLauncher::TryBuildUnrealInsightsExe(const FString& Path, const FString& LaunchParameters)
 {
diff --git a/Engine/Plugins/TraceUtilities/Source/EditorTraceUtilities/Private/UnrealInsightsLauncher.h b/Engine/Plugins/TraceUtilities/Source/EditorTraceUtilities/Private/UnrealInsightsLauncher.h
deleted file mode 100644
diff --git a/Engine/Plugins/TraceUtilities/Source/EditorTraceUtilities/Public/UnrealInsightsLauncher.h b/Engine/Plugins/TraceUtilities/Source/EditorTraceUtilities/Public/UnrealInsightsLauncher.h
new file mode 100644
--- /dev/null
+++ b/Engine/Plugins/TraceUtilities/Source/EditorTraceUtilities/Public/UnrealInsightsLauncher.h
@@ -0,0 +1,127 @@
+// Copyright Epic Games, Inc. All Rights Reserved.
+
+#pragma once
+
+#include "Async/TaskGraphFwd.h"
+#include "Async/TaskGraphInterfaces.h"
+#include "CoreMinimal.h"
+
+class EDITORTRACEUTILITIES_API FUnrealInsightsLauncher : public TSharedFromThis<FUnrealInsightsLauncher>
+{
+	friend class FLogMessageOnGameThreadTask;
+
+public:
+	FUnrealInsightsLauncher();
+	~FUnrealInsightsLauncher();
+
+	static TSharedPtr<FUnrealInsightsLauncher> Instance;
+	
+	static const TSharedPtr<FUnrealInsightsLauncher>& Get()
+	{
+		if (!Instance)
+		{
+			Instance = MakeShared<FUnrealInsightsLauncher>();
+		}
+		return Instance;
+	}
+
+	/**
+	 * Returns the full path to UnrealInsights.exe for the current engine installation 
+	 */
+	FString GetInsightsApplicationPath();
+
+	/**
+	 * Launches UnrealInsights.exe from the given Path, displays a editor message if it fails.
+	 * @param Path The full filename of UnrealInsights.exe to launch
+	 * @param Parameters The command line parameters to use when launching the exe
+	 */	
+	void StartUnrealInsights(const FString& Path, const FString& Parameters = TEXT(""));
+
+	/**
+	* Closes UnrealInsights.exe.
+	*/
+	void CloseUnrealInsights();
+
+	/**
+	 * Try to open a trace file from the trace target returned by UE::Trace::GetTraceDestination().
+	 * Destination may contain a file name or IP address of a trace server. In the latter case we try to query the
+	 * file name of the active/live trace through the storebrowser API and open the remote trace in Insights.
+	 * @param Destination The trace destination of the currently running trace session.
+	 */
+	bool TryOpenTraceFromDestination(const FString& Destination );
+
+	/**
+	 * Launches UnrealInsights.exe with FilePath as the only parameter
+	 * @param  FilePath Full file name and path to a .utrace file on the local machine
+	 * @return 
+	 */
+	bool OpenTraceFile(const FString& FilePath);
+
+	/**
+	 * Launches UnrealInsights.exe with the -Store and -OpenTraceID parameters
+	 * @param  TraceHostAddress hostname or ip address of the trace host
+	 * @param  TraceHostPort Port the Trace Host is running on
+	 * @param  TraceID ID of the Trace to open, can be obtained through the UE::Trace::FStoreBrowser interface
+	 * @return 
+	 */
+	bool OpenRemoteTrace(const FString& TraceHostAddress, const uint16 TraceHostPort, uint32 TraceID);
+
+	/**
+	* This function tries to query the file name of the active/live trace through the storebrowser API
+	 * @param Hostname hostname or ip of the target UnrealTraceServerX
+	 */
+	bool OpenActiveTraceFromStore(const FString& TraceHostAddress);
+
+private:
+	/*
+	 * Tries to build UnrealInsights via UAT, will launch with forwarded parameters if successful.
+	 * Assumes that the Insights Executable in path belongs to this engine.
+	 */
+	void TryBuildUnrealInsightsExe(const FString& Path, const FString& LaunchParameters = TEXT(""));
+
+	/// Logs an error message to the MessageLog window
+	void LogMessage(const FText& Message);
+
+	/// Logs an error message to the MessageLog window using a task that is ran on the Game Thread.
+	void LogMessageOnGameThread(const FText& Message);
+
+private:
+	/** The name of the Unreal Insights log listing. */
+	FName LogListingName;
+
+	/** The proccess handler of the Unreal Insights. */
+	FProcHandle UnrealInsightsHandle;
+};
+
+typedef TMap<FString, uint32> FLiveSessionsMap;
+
+struct FLiveSessionTaskData
+{
+	FLiveSessionsMap TaskLiveSessionData;
+	uint32 StorePort;
+};
+
+class FLiveSessionTracker
+{
+public:
+	FLiveSessionTracker();
+	~FLiveSessionTracker() {}
+
+	void Update();
+	void StartQuery();
+
+	bool HasData();
+	const FLiveSessionsMap& GetLiveSessions();
+	uint32 GetStorePort();
+
+private:
+	bool bHasData = false;
+	bool bIsQueryInProgress = false;
+	uint32 StorePort;
+
+	FLiveSessionsMap LiveSessionMap;
+	TSharedPtr<FLiveSessionTaskData> TaskLiveSessionData;
+
+	FGraphEventRef Event;
+};
+
