Fix StartUnrealInsights auto test

diff --git a/Engine/Plugins/TraceUtilities/Source/EditorTraceUtilities/Private/UnrealInsightsLauncher.cpp b/Engine/Plugins/TraceUtilities/Source/EditorTraceUtilities/Private/UnrealInsightsLauncher.cpp
--- a/Engine/Plugins/TraceUtilities/Source/EditorTraceUtilities/Private/UnrealInsightsLauncher.cpp
+++ b/Engine/Plugins/TraceUtilities/Source/EditorTraceUtilities/Private/UnrealInsightsLauncher.cpp
@@ -110,10 +110,16 @@ FString FUnrealInsightsLauncher::GetInsightsApplicationPath()
 }
 
 void FUnrealInsightsLauncher::StartUnrealInsights(const FString& Path, const FString& Parameters)
+{
+	auto Callback = [](const EStartInsightsResult Result) {};
+	StartUnrealInsights(Path, Parameters, Callback);
+}
+
+void FUnrealInsightsLauncher::StartUnrealInsights(const FString& Path, const FString& Parameters, StartUnrealInsightsCallback Callback)
 {
     if (!FPaths::FileExists(Path))
     {
-    	TryBuildUnrealInsightsExe(Path, Parameters);
+    	BuildUnrealInsights(Path, Parameters, Callback);
     	return;
     }
 	
@@ -133,12 +139,14 @@ void FUnrealInsightsLauncher::StartUnrealInsights(const FString& Path, const FSt
 	if (UnrealInsightsHandle.IsValid())
 	{
 		UE_LOG(LogTraceUtilities, Log, TEXT("Launched Unreal Insights executable: %s %s"), *Path, *Parameters);
+		Callback(EStartInsightsResult::Completed);
 	}
 	else
 	{
 		const FText	MessageBoxTextFmt = LOCTEXT("ExecutableNotFound_TextFmt", "Could not start Unreal Insights executable at path: {0}");
 		const FText MessageBoxText = FText::Format(MessageBoxTextFmt, FText::FromString(Path));
 		LogMessageOnGameThread(MessageBoxText);
+		Callback(EStartInsightsResult::LaunchFailed);
 	}
 }
 
@@ -154,7 +162,7 @@ void FUnrealInsightsLauncher::CloseUnrealInsights()
 	}
 }
 
-void FUnrealInsightsLauncher::TryBuildUnrealInsightsExe(const FString& Path, const FString& LaunchParameters)
+void FUnrealInsightsLauncher::BuildUnrealInsights(const FString& Path, const FString& LaunchParameters, StartUnrealInsightsCallback Callback)
 {
 	UE_LOG(LogTraceUtilities, Log, TEXT("Could not find the Unreal Insights executable: %s. Attempting to build UnrealInsights."), *Path);
 
@@ -171,18 +179,23 @@ void FUnrealInsightsLauncher::TryBuildUnrealInsightsExe(const FString& Path, con
 #endif
 
 	IUATHelperModule::Get().CreateUatTask(Arguments, PlatformName, LOCTEXT("BuildingUnrealInsights", "Building Unreal Insights"),
-		LOCTEXT("BuildUnrealInsightsTask", "Build Unreal Insights Task"), FAppStyle::GetBrush(TEXT("MainFrame.CookContent")), nullptr, [this, Path, LaunchParameters](FString Result, double Time)
+		LOCTEXT("BuildUnrealInsightsTask", "Build Unreal Insights Task"), FAppStyle::GetBrush(TEXT("MainFrame.CookContent")), nullptr,
+		[this, Path, LaunchParameters, Callback](FString Result, double Time)
 		{
 			if (Result.Equals(TEXT("Completed")))
 			{
 #if PLATFORM_MAC
 				// On Mac we genereate the path again so that it includes the newly built executable.
 				FString NewPath = GetInsightsApplicationPath();
-				this->StartUnrealInsights(NewPath, LaunchParameters);
+				this->StartUnrealInsights(NewPath, LaunchParameters, Callback);
 #else
-				this->StartUnrealInsights(Path, LaunchParameters);
+				this->StartUnrealInsights(Path, LaunchParameters, Callback);
 #endif
 			}
+			else
+			{
+				Callback(EStartInsightsResult::BuildFailed);
+			}
 		});
 }
 
diff --git a/Engine/Plugins/TraceUtilities/Source/EditorTraceUtilities/Public/UnrealInsightsLauncher.h b/Engine/Plugins/TraceUtilities/Source/EditorTraceUtilities/Public/UnrealInsightsLauncher.h
--- a/Engine/Plugins/TraceUtilities/Source/EditorTraceUtilities/Public/UnrealInsightsLauncher.h
+++ b/Engine/Plugins/TraceUtilities/Source/EditorTraceUtilities/Public/UnrealInsightsLauncher.h
@@ -11,6 +11,15 @@ class EDITORTRACEUTILITIES_API FUnrealInsightsLauncher : public TSharedFromThis<
 	friend class FLogMessageOnGameThreadTask;
 
 public:
+	enum class EStartInsightsResult : uint32
+	{
+		Completed = 0,
+		BuildFailed = 1,
+		LaunchFailed = 2,
+	};
+
+	typedef TFunction<void(const EStartInsightsResult /*Result*/)> StartUnrealInsightsCallback;
+
 	FUnrealInsightsLauncher();
 	~FUnrealInsightsLauncher();
 
@@ -31,12 +40,22 @@ public:
 	FString GetInsightsApplicationPath();
 
 	/**
-	 * Launches UnrealInsights.exe from the given Path, displays a editor message if it fails.
+	 * Launches the UnrealInsights executable from the given Path, displays an editor message if it fails.
+	 * If the executable is not found, a build process is started.
 	 * @param Path The full filename of UnrealInsights.exe to launch
 	 * @param Parameters The command line parameters to use when launching the exe
 	 */	
 	void StartUnrealInsights(const FString& Path, const FString& Parameters = TEXT(""));
 
+	/**
+	* Launches the UnrealInsights executable from the given Path, displays an editor message if it fails.
+	* If the executable is not found, a build process is started.
+	* @param Path The full filename of UnrealInsights.exe to launch
+	* @param Parameters The command line parameters to use when launching the exe
+	* @param Callback A Callback that will be called when the launch process is completed.
+	*/
+	void StartUnrealInsights(const FString& Path, const FString& Parameters, StartUnrealInsightsCallback Callback);
+
 	/**
 	* Closes UnrealInsights.exe.
 	*/
@@ -74,10 +93,10 @@ public:
 
 private:
 	/*
-	 * Tries to build UnrealInsights via UAT, will launch with forwarded parameters if successful.
+	 * Attempts building UnrealInsights via UAT, will launch with forwarded parameters if successful.
 	 * Assumes that the Insights Executable in path belongs to this engine.
 	 */
-	void TryBuildUnrealInsightsExe(const FString& Path, const FString& LaunchParameters = TEXT(""));
+	void BuildUnrealInsights(const FString& Path, const FString& LaunchParameters, StartUnrealInsightsCallback Callback);
 
 	/// Logs an error message to the MessageLog window
 	void LogMessage(const FText& Message);
