[Insights] Fixed memory leak in SInsightsStatusBarWidget::CacheTraceStorePath()

diff --git a/Engine/Plugins/TraceUtilities/Source/EditorTraceUtilities/Private/InsightsStatusBar/SInsightsStatusBar.cpp b/Engine/Plugins/TraceUtilities/Source/EditorTraceUtilities/Private/InsightsStatusBar/SInsightsStatusBar.cpp
--- a/Engine/Plugins/TraceUtilities/Source/EditorTraceUtilities/Private/InsightsStatusBar/SInsightsStatusBar.cpp
+++ b/Engine/Plugins/TraceUtilities/Source/EditorTraceUtilities/Private/InsightsStatusBar/SInsightsStatusBar.cpp
@@ -859,7 +859,9 @@ void SInsightsStatusBarWidget::CacheTraceStorePath()
 {
 	if (TraceStorePath.IsEmpty())
 	{
-		UE::Trace::FStoreClient* StoreClient = UE::Trace::FStoreClient::Connect(TEXT("localhost"));
+		using UE::Trace::FStoreClient;
+		FStoreClient* StoreClientPtr = FStoreClient::Connect(TEXT("localhost"));
+		TUniquePtr<FStoreClient> StoreClient = TUniquePtr<FStoreClient>(StoreClientPtr);
 
 		if (!StoreClient)
 		{
@@ -867,13 +869,12 @@ void SInsightsStatusBarWidget::CacheTraceStorePath()
 			return;
 		}
 
-		const UE::Trace::FStoreClient::FStatus* Status = StoreClient->GetStatus();
+		const FStoreClient::FStatus* Status = StoreClient->GetStatus();
 		if (!Status)
 		{
 			LogMessage(LOCTEXT("FailedToGetStoreStatusMsg", "Failed to get the status of the store client."));
 			return;
 		}
-
 		TraceStorePath = FString(Status->GetStoreDir());
 	}
 }
