Fix StartTraceSendTo blueprint
#prelight none

diff --git a/Engine/Plugins/TraceUtilities/Source/TraceUtilities/Private/TraceUtilLibrary.cpp b/Engine/Plugins/TraceUtilities/Source/TraceUtilities/Private/TraceUtilLibrary.cpp
--- a/Engine/Plugins/TraceUtilities/Source/TraceUtilities/Private/TraceUtilLibrary.cpp
+++ b/Engine/Plugins/TraceUtilities/Source/TraceUtilities/Private/TraceUtilLibrary.cpp
@@ -18,7 +18,7 @@ bool UTraceUtilLibrary::StartTraceToFile(const FString& FileName, const TArray<F
 bool UTraceUtilLibrary::StartTraceSendTo(const FString& Target, const TArray<FString>& Channels)
 {
 	FString ChannelString = FString::Join(Channels, TEXT(","));
-	return FTraceAuxiliary::Start(FTraceAuxiliary::EConnectionType::File ,*Target, *ChannelString);
+	return FTraceAuxiliary::Start(FTraceAuxiliary::EConnectionType::Network ,*Target, *ChannelString);
 }
 
 bool UTraceUtilLibrary::StopTracing()
