Add Trace Pause/Resume menu entry to the Trace Editor Widget.

diff --git a/Engine/Plugins/TraceUtilities/Source/EditorTraceUtilities/Private/InsightsStatusBar/EditorTraceUtilitiesStyle.h b/Engine/Plugins/TraceUtilities/Source/EditorTraceUtilities/Private/InsightsStatusBar/EditorTraceUtilitiesStyle.h
--- a/Engine/Plugins/TraceUtilities/Source/EditorTraceUtilities/Private/InsightsStatusBar/EditorTraceUtilitiesStyle.h
+++ b/Engine/Plugins/TraceUtilities/Source/EditorTraceUtilities/Private/InsightsStatusBar/EditorTraceUtilitiesStyle.h
@@ -52,6 +52,9 @@ public:
 		Set("Icons.TraceServerStart", new CORE_IMAGE_BRUSH_SVG("Starship/Common/play", Icon16x16, FStyleColors::AccentGreen));
 		Set("Icons.TraceServerStop", new CORE_IMAGE_BRUSH_SVG("Starship/Common/close", Icon16x16, FStyleColors::AccentRed));
 
+		Set("Icons.PauseTrace.Menu", new IMAGE_BRUSH_SVG("Pause_16", Icon16x16));
+		Set("Icons.ResumeTrace.Menu", new CORE_IMAGE_BRUSH_SVG("Starship/Common/play", Icon16x16));
+
 		FSlateStyleRegistry::RegisterSlateStyle(*this);
 	}
 
diff --git a/Engine/Plugins/TraceUtilities/Source/EditorTraceUtilities/Private/InsightsStatusBar/SInsightsStatusBar.cpp b/Engine/Plugins/TraceUtilities/Source/EditorTraceUtilities/Private/InsightsStatusBar/SInsightsStatusBar.cpp
--- a/Engine/Plugins/TraceUtilities/Source/EditorTraceUtilities/Private/InsightsStatusBar/SInsightsStatusBar.cpp
+++ b/Engine/Plugins/TraceUtilities/Source/EditorTraceUtilities/Private/InsightsStatusBar/SInsightsStatusBar.cpp
@@ -170,7 +170,7 @@ void SInsightsStatusBarWidget::Construct(const FArguments& InArgs)
 			.HAlign(HAlign_Left)
 			.VAlign(VAlign_Bottom)
 			.ToolTipText(this, &SInsightsStatusBarWidget::GetRecordingButtonTooltipText)
-			.OnClicked_Lambda([this]() { this->ToggleTracing_OnClicked(); return FReply::Handled(); })
+			.OnClicked_Lambda([this]() { this->ToggleTrace_OnClicked(); return FReply::Handled(); })
 			.OnHovered_Lambda([this]() { this->bIsTraceRecordButtonHovered = true; })
 			.OnUnhovered_Lambda([this]() { this->bIsTraceRecordButtonHovered = false; })
 			.Content()
@@ -331,11 +331,36 @@ TSharedRef<SWidget> SInsightsStatusBarWidget::MakeTraceMenu()
 			TAttribute<FText>::CreateSP(this, &SInsightsStatusBarWidget::GetTraceMenuItemText),
 			TAttribute<FText>::CreateSP(this, &SInsightsStatusBarWidget::GetTraceMenuItemTooltipText),
 			FSlateIcon(FEditorTraceUtilitiesStyle::Get().GetStyleSetName(), "Icons.StartTrace.Menu"),
-			FUIAction(FExecuteAction::CreateSP(this, &SInsightsStatusBarWidget::ToggleTracing_OnClicked)),
+			FUIAction(FExecuteAction::CreateSP(this, &SInsightsStatusBarWidget::ToggleTrace_OnClicked)),
 			NAME_None,
 			EUserInterfaceActionType::Button
 		);
 
+		if (FTraceAuxiliary::IsPaused())
+		{
+			MenuBuilder.AddMenuEntry(
+				LOCTEXT("ResumeTraceButtonText", "Resume Trace"),
+				LOCTEXT("ResumesTraceButtonTooltip", "Enables all channels that were active when tracing was paused."),
+				FSlateIcon(FEditorTraceUtilitiesStyle::Get().GetStyleSetName(), "Icons.ResumeTrace.Menu"),
+				FUIAction(FExecuteAction::CreateSP(this, &SInsightsStatusBarWidget::TogglePauseTrace_OnClicked),
+						  FCanExecuteAction::CreateSP(this, &SInsightsStatusBarWidget::PauseTrace_CanExecute)),
+				NAME_None,
+				EUserInterfaceActionType::Button
+			);
+		}
+		else
+		{
+			MenuBuilder.AddMenuEntry(
+				LOCTEXT("PauseTraceButtonText", "Pause Trace"),
+				TAttribute<FText>::CreateSP(this, &SInsightsStatusBarWidget::GetPauseTraceMenuItemTooltipText),
+				FSlateIcon(FEditorTraceUtilitiesStyle::Get().GetStyleSetName(), "Icons.PauseTrace.Menu"),
+				FUIAction(FExecuteAction::CreateSP(this, &SInsightsStatusBarWidget::TogglePauseTrace_OnClicked),
+						  FCanExecuteAction::CreateSP(this, &SInsightsStatusBarWidget::PauseTrace_CanExecute)),
+				NAME_None,
+				EUserInterfaceActionType::Button
+			);
+		}
+
 		MenuBuilder.AddMenuEntry(
 			LOCTEXT("SaveSnapshotLabel", "Save Trace Snapshot"),
 			LOCTEXT("SaveSnapshotTooltip", "Save the current trace buffer to the selected trace destination."),
@@ -744,7 +769,7 @@ FText SInsightsStatusBarWidget::GetTraceMenuItemTooltipText() const
 	return LOCTEXT("StartTraceButtonTooltip", "Start tracing to the selected trace destination.");
 }
 
-void SInsightsStatusBarWidget::ToggleTracing_OnClicked()
+void SInsightsStatusBarWidget::ToggleTrace_OnClicked()
 {
 	if (UE::Trace::IsTracing())
 	{
@@ -774,6 +799,38 @@ void SInsightsStatusBarWidget::ToggleTracing_OnClicked()
 	}
 }
 
+bool SInsightsStatusBarWidget::PauseTrace_CanExecute()
+{
+	return UE::Trace::IsTracing();
+}
+
+FText SInsightsStatusBarWidget::GetPauseTraceMenuItemTooltipText() const
+{
+	if (!UE::Trace::IsTracing())
+	{
+		return LOCTEXT("PauseTraceDisabledButtonTooltip", "Tracing must be running to enable the pause functionality.");
+	}
+
+	return LOCTEXT("PauseTraceButtonTooltip", "Disables all enabled trace channels. The same channels will be re-enabled when tracing is resumed.");
+}
+
+void SInsightsStatusBarWidget::TogglePauseTrace_OnClicked()
+{ 
+	if (!UE::Trace::IsTracing())
+	{
+		return;
+	}
+
+	if (FTraceAuxiliary::IsPaused())
+	{
+		FTraceAuxiliary::Resume();
+	}
+	else
+	{
+		FTraceAuxiliary::Pause();
+	}
+}
+
 bool SInsightsStatusBarWidget::StartTracing()
 {
 	if (TraceDestination == ETraceDestination::TraceStore)
diff --git a/Engine/Plugins/TraceUtilities/Source/EditorTraceUtilities/Private/InsightsStatusBar/SInsightsStatusBar.h b/Engine/Plugins/TraceUtilities/Source/EditorTraceUtilities/Private/InsightsStatusBar/SInsightsStatusBar.h
--- a/Engine/Plugins/TraceUtilities/Source/EditorTraceUtilities/Private/InsightsStatusBar/SInsightsStatusBar.h
+++ b/Engine/Plugins/TraceUtilities/Source/EditorTraceUtilities/Private/InsightsStatusBar/SInsightsStatusBar.h
@@ -92,8 +92,11 @@ private:
 
 	FText GetTraceMenuItemText() const;
 	FText GetTraceMenuItemTooltipText() const;
+	void ToggleTrace_OnClicked();
 
-	void ToggleTracing_OnClicked();
+	bool PauseTrace_CanExecute();
+	FText GetPauseTraceMenuItemTooltipText() const;
+	void TogglePauseTrace_OnClicked();
 
 	EVisibility GetStartTraceIconVisibility() const;
 	EVisibility GetStopTraceIconVisibility() const;
diff --git a/Engine/Source/Runtime/Core/Private/ProfilingDebugging/TraceAuxiliary.cpp b/Engine/Source/Runtime/Core/Private/ProfilingDebugging/TraceAuxiliary.cpp
--- a/Engine/Source/Runtime/Core/Private/ProfilingDebugging/TraceAuxiliary.cpp
+++ b/Engine/Source/Runtime/Core/Private/ProfilingDebugging/TraceAuxiliary.cpp
@@ -112,6 +112,7 @@ public:
 	bool Stop();
 	void ResumeChannels();
 	void PauseChannels();
+	bool IsPaused();
 	void EnableCommandlineChannels();
 	void EnableCommandlineChannelsPostInitialize();
 	void SetTruncateFile(bool bTruncateFile);
@@ -369,6 +370,7 @@ bool FTraceAuxiliaryImpl::Stop()
 
 	FString StopedDest;
 	FTraceAuxiliary::EConnectionType StopedType = FTraceAuxiliary::EConnectionType::None;
+	PausedPreset.Empty();
 
 	{
 		FWriteScopeLock _(CurrentTargetLock);
@@ -420,6 +422,8 @@ void FTraceAuxiliaryImpl::ResumeChannels()
 {
 	// Enable channels from the "paused" preset.
 	ForEachChannel(*PausedPreset, false, &FTraceAuxiliaryImpl::EnableChannel);
+
+	PausedPreset.Empty();
 }
 
 ////////////////////////////////////////////////////////////////////////////////
@@ -436,6 +440,12 @@ void FTraceAuxiliaryImpl::PauseChannels()
 	ForEachChannel(*PausedPreset, true, &FTraceAuxiliaryImpl::DisableChannel);
 }
 
+////////////////////////////////////////////////////////////////////////////////
+bool FTraceAuxiliaryImpl::IsPaused()
+{
+	return !PausedPreset.IsEmpty();
+}
+
 ////////////////////////////////////////////////////////////////////////////////
 void FTraceAuxiliaryImpl::EnableCommandlineChannels()
 {
@@ -1333,6 +1343,16 @@ bool FTraceAuxiliary::Pause()
 	return true;
 }
 
+////////////////////////////////////////////////////////////////////////////////
+bool FTraceAuxiliary::IsPaused()
+{
+#if UE_TRACE_ENABLED
+	return GTraceAuxiliary.IsPaused();
+#else
+	return false;
+#endif
+}
+
 ////////////////////////////////////////////////////////////////////////////////
 bool FTraceAuxiliary::Resume()
 {
diff --git a/Engine/Source/Runtime/Core/Public/ProfilingDebugging/TraceAuxiliary.h b/Engine/Source/Runtime/Core/Public/ProfilingDebugging/TraceAuxiliary.h
--- a/Engine/Source/Runtime/Core/Public/ProfilingDebugging/TraceAuxiliary.h
+++ b/Engine/Source/Runtime/Core/Public/ProfilingDebugging/TraceAuxiliary.h
@@ -98,6 +98,11 @@ public:
 	 */
 	static CORE_API bool Pause();
 
+	/**
+	* @return True if trace was paused and the list of channels to resume exists.
+	*/
+	static CORE_API bool IsPaused();
+
 	/**
 	 * Resume tracing by enabling all previously active channels.
 	 */
