Expose ETraceDestination which is set by the Insights tracing status bar combo button.
- This is needed to allow external modules that want to start a trace re-use the same setting a user has already set.
- Specifically,  this is needed for ConcertInsights which will be starting synchronized traces across multiple clients in Multi User sessions.
- Also added some of the code in UE::EditorTraceUtilities namespace as required by code guidelines

diff --git a/Engine/Plugins/TraceUtilities/Source/EditorTraceUtilities/Private/EditorTraceUtilities.cpp b/Engine/Plugins/TraceUtilities/Source/EditorTraceUtilities/Private/EditorTraceUtilities.cpp
--- a/Engine/Plugins/TraceUtilities/Source/EditorTraceUtilities/Private/EditorTraceUtilities.cpp
+++ b/Engine/Plugins/TraceUtilities/Source/EditorTraceUtilities/Private/EditorTraceUtilities.cpp
@@ -4,14 +4,15 @@
 
 #include "Misc/ConfigContext.h"
 #include "ToolMenus.h"
+#include "InsightsStatusBar/SInsightsStatusBar.h"
 
 DEFINE_LOG_CATEGORY(LogTraceUtilities)
 
 #define LOCTEXT_NAMESPACE "FEditorTraceUtilitiesModule"
 
-/**
-  * This function will add the SInsightsStatusBarWidget to the Editor's status bar at the bottom ("LevelEditor.StatusBar.ToolBar").
-  */
+namespace UE::EditorTraceUtilities
+{
+/** This function will add the SInsightsStatusBarWidget to the Editor's status bar at the bottom ("LevelEditor.StatusBar.ToolBar"). */
 void RegisterInsightsStatusWidgetWithToolMenu();
 
 FString FEditorTraceUtilitiesModule::EditorTraceUtilitiesIni;
@@ -21,7 +22,6 @@ void FEditorTraceUtilitiesModule::StartupModule()
 	LLM_SCOPE_BYNAME(TEXT("Insights"));
 
 	RegisterInsightsStatusWidgetWithToolMenu();
-
 	FConfigContext::ReadIntoGConfig().Load(TEXT("TraceUtilities"), EditorTraceUtilitiesIni);
 }
 
@@ -31,6 +31,12 @@ void FEditorTraceUtilitiesModule::ShutdownModule()
 	UToolMenus::UnRegisterStartupCallback(RegisterStartupCallbackHandle);
 }
 
+const FStatusBarTraceSettings& FEditorTraceUtilitiesModule::GetTraceSettings() const
+{
+	return SInsightsStatusBarWidget::StatusBarTraceSettings;
+}
+} // namespace UE::EditorTraceUtilities
+
 #undef LOCTEXT_NAMESPACE
     
-IMPLEMENT_MODULE(FEditorTraceUtilitiesModule, EditorTraceUtilities)
\ No newline at end of file
+IMPLEMENT_MODULE(UE::EditorTraceUtilities::FEditorTraceUtilitiesModule, EditorTraceUtilities)
\ No newline at end of file
diff --git a/Engine/Plugins/TraceUtilities/Source/EditorTraceUtilities/Private/EditorTraceUtilities.h b/Engine/Plugins/TraceUtilities/Source/EditorTraceUtilities/Private/EditorTraceUtilities.h
new file mode 100644
--- /dev/null
+++ b/Engine/Plugins/TraceUtilities/Source/EditorTraceUtilities/Private/EditorTraceUtilities.h
@@ -0,0 +1,29 @@
+﻿// Copyright Epic Games, Inc. All Rights Reserved.
+
+#pragma once
+
+#include "CoreMinimal.h"
+#include "IEditorTraceUtilitiesModule.h"
+#include "Modules/ModuleManager.h"
+
+DECLARE_LOG_CATEGORY_EXTERN(LogTraceUtilities, Log, All)
+
+namespace UE::EditorTraceUtilities
+{
+class FEditorTraceUtilitiesModule : public IEditorTraceUtilitiesModule
+{
+public:
+
+	//~ Begin IEditorTraceUtilitiesModule Interface
+	virtual void StartupModule() override;
+	virtual void ShutdownModule() override;
+	virtual const FStatusBarTraceSettings& GetTraceSettings() const override;
+	//~ End IEditorTraceUtilitiesModule Interface
+
+	static const FString& GetTraceUtilitiesIni() { return EditorTraceUtilitiesIni; }
+
+private:
+	FDelegateHandle RegisterStartupCallbackHandle;
+	static FString EditorTraceUtilitiesIni;
+};
+} // namespace UE::EditorTraceUtilities
\ No newline at end of file
diff --git a/Engine/Plugins/TraceUtilities/Source/EditorTraceUtilities/Private/InsightsStatusBar/SInsightsStatusBar.cpp b/Engine/Plugins/TraceUtilities/Source/EditorTraceUtilities/Private/InsightsStatusBar/SInsightsStatusBar.cpp
--- a/Engine/Plugins/TraceUtilities/Source/EditorTraceUtilities/Private/InsightsStatusBar/SInsightsStatusBar.cpp
+++ b/Engine/Plugins/TraceUtilities/Source/EditorTraceUtilities/Private/InsightsStatusBar/SInsightsStatusBar.cpp
@@ -40,6 +40,10 @@
 
 #define LOCTEXT_NAMESPACE "InsightsEditor"
 
+namespace UE::EditorTraceUtilities
+{
+FStatusBarTraceSettings SInsightsStatusBarWidget::StatusBarTraceSettings;
+
 const TCHAR* SInsightsStatusBarWidget::DefaultPreset = TEXT("default");
 const TCHAR* SInsightsStatusBarWidget::MemoryPreset= TEXT("default,memory");
 const TCHAR* SInsightsStatusBarWidget::TaskGraphPreset = TEXT("default,task");
@@ -223,11 +227,11 @@ void SInsightsStatusBarWidget::Construct(const FArguments& InArgs)
 
 	if (FTraceAuxiliary::GetConnectionType() == FTraceAuxiliary::EConnectionType::Network)
 	{
-		TraceDestination = ETraceDestination::TraceStore;
+		StatusBarTraceSettings.TraceDestination = ETraceDestination::TraceStore;
 	}
 	if (FTraceAuxiliary::GetConnectionType() == FTraceAuxiliary::EConnectionType::File)
 	{
-		TraceDestination = ETraceDestination::File;
+		StatusBarTraceSettings.TraceDestination = ETraceDestination::File;
 	}
 
 	LogListingName = TEXT("UnrealInsights");
@@ -703,12 +707,12 @@ FString SInsightsStatusBarWidget::GetLatestTraceFileFromFolder(const FString& In
 
 void SInsightsStatusBarWidget::SetTraceDestination_Execute(ETraceDestination InDestination)
 {
-	TraceDestination = InDestination;
+	StatusBarTraceSettings.TraceDestination = InDestination;
 }
 
 bool SInsightsStatusBarWidget::SetTraceDestination_IsChecked(ETraceDestination InDestination)
 {
-	return InDestination == TraceDestination;
+	return InDestination == StatusBarTraceSettings.TraceDestination;
 }
 
 bool SInsightsStatusBarWidget::SetTraceDestination_CanExecute()
@@ -723,7 +727,7 @@ bool SInsightsStatusBarWidget::SetTraceDestination_CanExecute()
 
 void SInsightsStatusBarWidget::SaveSnapshot()
 {
-	if (TraceDestination == ETraceDestination::File)
+	if (StatusBarTraceSettings.TraceDestination == ETraceDestination::File)
 	{
 		const bool bResult = FTraceAuxiliary::WriteSnapshot(nullptr);
 		if (bResult)
@@ -833,11 +837,11 @@ void SInsightsStatusBarWidget::TogglePauseTrace_OnClicked()
 
 bool SInsightsStatusBarWidget::StartTracing()
 {
-	if (TraceDestination == ETraceDestination::TraceStore)
+	if (StatusBarTraceSettings.TraceDestination == ETraceDestination::TraceStore)
 	{
 		return FTraceAuxiliary::Start(FTraceAuxiliary::EConnectionType::Network, TEXT("localhost"), nullptr);
 	}
-	else if (TraceDestination == ETraceDestination::File)
+	else if (StatusBarTraceSettings.TraceDestination == ETraceDestination::File)
 	{
 		return FTraceAuxiliary::Start(FTraceAuxiliary::EConnectionType::File, nullptr, nullptr);
 	}
@@ -1115,5 +1119,6 @@ void SInsightsStatusBarWidget::OpenTrace(int32 Index)
 		FUnrealInsightsLauncher::Get()->TryOpenTraceFromDestination(Traces[Index]->FilePath);
 	}
 }
+} // namespace UE::EditorTraceUtilities
 
 #undef LOCTEXT_NAMESPACE
diff --git a/Engine/Plugins/TraceUtilities/Source/EditorTraceUtilities/Private/InsightsStatusBar/SInsightsStatusBar.h b/Engine/Plugins/TraceUtilities/Source/EditorTraceUtilities/Private/InsightsStatusBar/SInsightsStatusBar.h
--- a/Engine/Plugins/TraceUtilities/Source/EditorTraceUtilities/Private/InsightsStatusBar/SInsightsStatusBar.h
+++ b/Engine/Plugins/TraceUtilities/Source/EditorTraceUtilities/Private/InsightsStatusBar/SInsightsStatusBar.h
@@ -2,7 +2,7 @@
 
 #pragma once
 
-#include "CoreMinimal.h"
+#include "EditorTraceUtilities.h"
 #include "ProfilingDebugging/TraceAuxiliary.h"
 #include "Widgets/DeclarativeSyntaxSupport.h"
 #include "Widgets/SWidget.h"
@@ -15,6 +15,8 @@ class STraceServerControl;
 
 TSharedRef<SWidget> CreateInsightsStatusBarWidget();
 
+namespace UE::EditorTraceUtilities
+{
 struct FTraceFileInfo
 {
 	FString FilePath;
@@ -33,13 +35,6 @@ struct FTraceFileInfo
  */
 class SInsightsStatusBarWidget : public SCompoundWidget
 {
-private:
-	enum class ETraceDestination : uint32
-	{
-		TraceStore = 0,
-		File = 1
-	};
-
 	struct FChannelData
 	{
 		FString Name;
@@ -55,6 +50,10 @@ private:
 	};
 
 public:
+
+	/** Settings this widget uses. */
+	static FStatusBarTraceSettings StatusBarTraceSettings;
+	
 	SLATE_BEGIN_ARGS(SInsightsStatusBarWidget) {}
 	SLATE_END_ARGS()
 
@@ -110,9 +109,6 @@ private:
 	void LogMessage(const FText& Text);
 	void ShowNotification(const FText& Text, const FText& SubText);
 
-	void SetTraceChannels(const TCHAR* InChannels);
-	bool IsPresetSet(const TCHAR* InChannels) const;
-
 	bool GetBooleanSettingValue(const TCHAR* InSettingName);
 	void ToggleBooleanSettingValue(const TCHAR* InSettingName);
 
@@ -151,7 +147,6 @@ private:
 	static const TCHAR* OpenInsightsAfterTraceSettingName;
 	static const TCHAR* ShowInExplorerAfterTraceSettingName;
 
-	ETraceDestination TraceDestination = ETraceDestination::TraceStore;
 	bool bIsTraceRecordButtonHovered = false;
 	mutable double ConnectionStartTime = 0.0f;
 
@@ -169,3 +164,4 @@ private:
 	TArray<TSharedPtr<FTraceFileInfo>> Traces;
 	FName LogListingName;
 };
+}
diff --git a/Engine/Plugins/TraceUtilities/Source/EditorTraceUtilities/Private/InsightsStatusBar/SRecentTracesList.cpp b/Engine/Plugins/TraceUtilities/Source/EditorTraceUtilities/Private/InsightsStatusBar/SRecentTracesList.cpp
--- a/Engine/Plugins/TraceUtilities/Source/EditorTraceUtilities/Private/InsightsStatusBar/SRecentTracesList.cpp
+++ b/Engine/Plugins/TraceUtilities/Source/EditorTraceUtilities/Private/InsightsStatusBar/SRecentTracesList.cpp
@@ -11,10 +11,11 @@
 #include "SInsightsStatusBar.h"
 #include "UnrealInsightsLauncher.h"
 #include "Widgets/Input/SButton.h"
-#include "Widgets/Layout/SScrollBorder.h"
 
 #define LOCTEXT_NAMESPACE "RecentTracesList"
 
+namespace UE::EditorTraceUtilities
+{
 void SRecentTracesListEntry::Construct(const FArguments& InArgs, TSharedPtr<FTraceFileInfo> InTrace, const FString& InStorePath, TSharedPtr<FLiveSessionTracker> InLiveSessionTracker)
 {
 	TraceInfo = InTrace;
@@ -130,5 +131,6 @@ EVisibility SRecentTracesListEntry::GetLiveLabelVisibility() const
 
 	return EVisibility::Collapsed;
 }
+} // namespace UE::EditorTraceUtilities
 
 #undef LOCTEXT_NAMESPACE
diff --git a/Engine/Plugins/TraceUtilities/Source/EditorTraceUtilities/Private/InsightsStatusBar/SRecentTracesList.h b/Engine/Plugins/TraceUtilities/Source/EditorTraceUtilities/Private/InsightsStatusBar/SRecentTracesList.h
--- a/Engine/Plugins/TraceUtilities/Source/EditorTraceUtilities/Private/InsightsStatusBar/SRecentTracesList.h
+++ b/Engine/Plugins/TraceUtilities/Source/EditorTraceUtilities/Private/InsightsStatusBar/SRecentTracesList.h
@@ -5,9 +5,12 @@
 #include "Widgets/SCompoundWidget.h"
 #include "Widgets/Views/SListView.h"
 
-struct FTraceFileInfo;
 class FLiveSessionTracker;
 
+namespace UE::EditorTraceUtilities
+{
+struct FTraceFileInfo;
+
 /**
  *  A widget that shows an entry in the recent traces submenu.
  */
@@ -29,3 +32,5 @@ private:
 	FString TraceName;
 	TSharedPtr<FLiveSessionTracker> LiveSessionTracker;
 };
+}
+
diff --git a/Engine/Plugins/TraceUtilities/Source/EditorTraceUtilities/Public/EditorTraceUtilities.h b/Engine/Plugins/TraceUtilities/Source/EditorTraceUtilities/Public/EditorTraceUtilities.h
deleted file mode 100644
diff --git a/Engine/Plugins/TraceUtilities/Source/EditorTraceUtilities/Public/IEditorTraceUtilitiesModule.h b/Engine/Plugins/TraceUtilities/Source/EditorTraceUtilities/Public/IEditorTraceUtilitiesModule.h
new file mode 100644
--- /dev/null
+++ b/Engine/Plugins/TraceUtilities/Source/EditorTraceUtilities/Public/IEditorTraceUtilitiesModule.h
@@ -0,0 +1,51 @@
+// Copyright Epic Games, Inc. All Rights Reserved.
+
+#pragma once
+
+#include "Modules/ModuleManager.h"
+
+namespace UE::EditorTraceUtilities
+{
+enum class ETraceDestination : uint32
+{
+	TraceStore = 0,
+	File = 1
+};
+
+/** The tracing settings the status bar extension manages. */
+struct FStatusBarTraceSettings
+{
+	ETraceDestination TraceDestination = ETraceDestination::TraceStore;
+};
+
+class EDITORTRACEUTILITIES_API IEditorTraceUtilitiesModule : public IModuleInterface
+{
+public:
+
+	/**
+	 * Singleton-like access to this module's interface.  This is just for convenience!
+	 * Beware of calling this during the shutdown phase, though.  Your module might have been unloaded already.
+	 *
+	 * @return Returns singleton instance, loading the module on demand if needed
+	 */
+	static inline IEditorTraceUtilitiesModule& Get()
+	{
+		static const FName ModuleName = "EditorTraceUtilities";
+		return FModuleManager::LoadModuleChecked<IEditorTraceUtilitiesModule>(ModuleName);
+	}
+
+	/**
+	 * Checks to see if this module is loaded and ready.  It is only valid to call Get() during shutdown if IsAvailable() returns true.
+	 *
+	 * @return True if the module is loaded and ready to use
+	 */
+	static inline bool IsAvailable()
+	{
+		static const FName ModuleName = "EditorTraceUtilities";
+		return FModuleManager::Get().IsModuleLoaded(ModuleName);
+	}
+
+	/** Gets the tracing settings the status bar extension manages. */
+	virtual const FStatusBarTraceSettings& GetTraceSettings() const = 0;
+};
+} // namespace UE::EditorTraceUtilities
\ No newline at end of file
