[Insights] Fixed "Trace Bookmark" (from Trace/Insights StatusBar Widget in Editor) to generate unique bookmarks (timestamps) instead of duplicating text of the first generated bookmark.

diff --git a/Engine/Plugins/TraceUtilities/Source/EditorTraceUtilities/Private/InsightsStatusBar/SInsightsStatusBar.cpp b/Engine/Plugins/TraceUtilities/Source/EditorTraceUtilities/Private/InsightsStatusBar/SInsightsStatusBar.cpp
--- a/Engine/Plugins/TraceUtilities/Source/EditorTraceUtilities/Private/InsightsStatusBar/SInsightsStatusBar.cpp
+++ b/Engine/Plugins/TraceUtilities/Source/EditorTraceUtilities/Private/InsightsStatusBar/SInsightsStatusBar.cpp
@@ -1059,8 +1059,8 @@ bool SInsightsStatusBarWidget::TraceBookmark_CanExecute()
 
 void SInsightsStatusBarWidget::TraceBookmark_Execute()
 {
-	FString Bookmark = FDateTime::Now().ToString(TEXT("Bookmark_%Y%m%d_%H%M%S"));
-	TRACE_BOOKMARK(*Bookmark);
+	const FString Bookmark = FDateTime::Now().ToString(TEXT("Bookmark_%Y%m%d_%H%M%S"));
+	TRACE_BOOKMARK(TEXT("%s"), *Bookmark);
 }
 
 void SInsightsStatusBarWidget::PopulateRecentTracesList()
