Fix trace status tooltip message in the Trace Status Bar Widget when not tracing.

diff --git a/Engine/Plugins/TraceUtilities/Source/EditorTraceUtilities/Private/InsightsStatusBar/SInsightsStatusBar.cpp b/Engine/Plugins/TraceUtilities/Source/EditorTraceUtilities/Private/InsightsStatusBar/SInsightsStatusBar.cpp
--- a/Engine/Plugins/TraceUtilities/Source/EditorTraceUtilities/Private/InsightsStatusBar/SInsightsStatusBar.cpp
+++ b/Engine/Plugins/TraceUtilities/Source/EditorTraceUtilities/Private/InsightsStatusBar/SInsightsStatusBar.cpp
@@ -542,7 +542,7 @@ FText SInsightsStatusBarWidget::GetTitleToolTipText() const
 
 	const FString Dest = FTraceAuxiliary::GetTraceDestinationString();
 
-	if (*Dest != 0)
+	if (!Dest.IsEmpty())
 	{
 		DescBuilder.AppendLineFormat(LOCTEXT("TracingToText", "Tracing to: {0}"), FText::FromString(Dest));
 	}
