Adding Trace Screenshot BP node to the Trace Utilities plugin

diff --git a/Engine/Plugins/TraceUtilities/Source/TraceUtilities/Private/TraceUtilLibrary.cpp b/Engine/Plugins/TraceUtilities/Source/TraceUtilities/Private/TraceUtilLibrary.cpp
--- a/Engine/Plugins/TraceUtilities/Source/TraceUtilities/Private/TraceUtilLibrary.cpp
+++ b/Engine/Plugins/TraceUtilities/Source/TraceUtilities/Private/TraceUtilLibrary.cpp
@@ -4,6 +4,7 @@
 
 #include "ProfilingDebugging/TraceAuxiliary.h"
 #include "Modules/ModuleManager.h"
+#include "ProfilingDebugging/TraceScreenshot.h"
 
 #include UE_INLINE_GENERATED_CPP_BY_NAME(TraceUtilLibrary)
 
@@ -87,3 +88,8 @@ void UTraceUtilLibrary::TraceMarkRegionEnd(const FString& Name)
 {
 	TRACE_END_REGION(*Name);
 }
+
+void UTraceUtilLibrary::TraceScreenshot(const FString& Name, bool bShowUI)
+{
+	FTraceScreenshot::RequestScreenshot(Name, bShowUI);
+}
diff --git a/Engine/Plugins/TraceUtilities/Source/TraceUtilities/Public/TraceUtilLibrary.h b/Engine/Plugins/TraceUtilities/Source/TraceUtilities/Public/TraceUtilLibrary.h
--- a/Engine/Plugins/TraceUtilities/Source/TraceUtilities/Public/TraceUtilLibrary.h
+++ b/Engine/Plugins/TraceUtilities/Source/TraceUtilities/Public/TraceUtilLibrary.h
@@ -58,4 +58,10 @@ public:
 	 */
 	UFUNCTION(BlueprintCallable, Category = "Perf | Insights Trace")
 	static void TraceMarkRegionEnd(const FString& Name);
+
+	/**
+	 * Triggers an Unreal Insights screenshot
+	 */
+	UFUNCTION(BlueprintCallable, Category = "Perf | Insights Trace")
+	static void TraceScreenshot(const FString& Name, bool bShowUI);
 };
