Adding controls to begin/end a trace region to the Insights status bar

diff --git a/Engine/Plugins/TraceUtilities/Source/EditorTraceUtilities/Private/InsightsStatusBar/EditorTraceUtilitiesStyle.h b/Engine/Plugins/TraceUtilities/Source/EditorTraceUtilities/Private/InsightsStatusBar/EditorTraceUtilitiesStyle.h
--- a/Engine/Plugins/TraceUtilities/Source/EditorTraceUtilities/Private/InsightsStatusBar/EditorTraceUtilitiesStyle.h
+++ b/Engine/Plugins/TraceUtilities/Source/EditorTraceUtilities/Private/InsightsStatusBar/EditorTraceUtilitiesStyle.h
@@ -49,6 +49,9 @@ public:
 		Set("Icons.PauseTrace.Menu", new IMAGE_BRUSH_SVG("Pause_16", Icon16x16));
 		Set("Icons.ResumeTrace.Menu", new CORE_IMAGE_BRUSH_SVG("Starship/Common/play", Icon16x16));
 
+		Set("Icons.BeginRegion.Menu", new CORE_IMAGE_BRUSH_SVG("Starship/Common/play", Icon16x16));
+		Set("Icons.EndRegion.Menu", new CORE_IMAGE_BRUSH_SVG("Starship/Common/stop", Icon16x16));
+
 		FSlateStyleRegistry::RegisterSlateStyle(*this);
 	}
 
diff --git a/Engine/Plugins/TraceUtilities/Source/EditorTraceUtilities/Private/InsightsStatusBar/SInsightsStatusBar.cpp b/Engine/Plugins/TraceUtilities/Source/EditorTraceUtilities/Private/InsightsStatusBar/SInsightsStatusBar.cpp
--- a/Engine/Plugins/TraceUtilities/Source/EditorTraceUtilities/Private/InsightsStatusBar/SInsightsStatusBar.cpp
+++ b/Engine/Plugins/TraceUtilities/Source/EditorTraceUtilities/Private/InsightsStatusBar/SInsightsStatusBar.cpp
@@ -60,6 +60,7 @@ const TCHAR* SInsightsStatusBarWidget::ContextSwitchesPreset = TEXT("default,con
 const TCHAR* SInsightsStatusBarWidget::SettingsCategory = TEXT("EditorTraceUtilities");
 const TCHAR* SInsightsStatusBarWidget::OpenLiveSessionOnTraceStartSettingName = TEXT("OpenLiveSessionOnTraceStart");
 const TCHAR* SInsightsStatusBarWidget::OpenInsightsAfterTraceSettingName = TEXT("OpenInsightsAfterTrace");
+const TCHAR* SInsightsStatusBarWidget::TraceRegionSettingName = TEXT("InsightsToolbarTraceRegion");
 const TCHAR* SInsightsStatusBarWidget::ShowInExplorerAfterTraceSettingName = TEXT("ShowInExplorerAfterTrace");
 
 class FOpenLiveSessionTask
@@ -268,7 +269,27 @@ TSharedRef<SWidget> SInsightsStatusBarWidget::MakeTraceMenu()
 			TAttribute<FText>(),
 			FSlateIcon(FEditorTraceUtilitiesStyle::Get().GetStyleSetName(), "Icons.Bookmark.Menu")
 		);
-
+		
+		MenuBuilder.AddEditableText(LOCTEXT("TraceRegionNameLabel", "Region Name:"),
+							GetTraceRegionNameDesc(),
+									FSlateIcon(),
+									TAttribute<FText>::Create(TAttribute<FText>::FGetter::CreateRaw(this, &SInsightsStatusBarWidget::GetTraceRegionName)),
+									FOnTextCommitted::CreateLambda([](const FText& NewRegionName, ETextCommit::Type)
+									{
+										GConfig->SetText(SettingsCategory, TraceRegionSettingName, NewRegionName, FEditorTraceUtilitiesModule::GetTraceUtilitiesIni());
+									}),
+									FOnTextChanged(),
+									RegionIsActive());
+		
+		MenuBuilder.AddMenuEntry(
+			GetRegionSwitchLabelText(),
+			GetRegionSwitchDescText(),
+			FSlateIcon(FEditorTraceUtilitiesStyle::Get().GetStyleSetName(), RegionIsActive() ? "Icons.EndRegion.Menu" : "Icons.BeginRegion.Menu"),
+			FUIAction(FExecuteAction::CreateSP(this, &SInsightsStatusBarWidget::ToggleRegion_Execute)),
+			NAME_None,
+			EUserInterfaceActionType::Button
+		);
+		
 		MenuBuilder.AddMenuEntry(
 			LOCTEXT("StatNamedEventsLabel", "Stat Named Events"),
 			LOCTEXT("StatNamedEventsDesc", "Enable or disable named events in the stats system."),
@@ -999,6 +1020,52 @@ void SInsightsStatusBarWidget::TraceBookmark_Execute()
 	TRACE_BOOKMARK(TEXT("%s"), *Bookmark);
 }
 
+FText SInsightsStatusBarWidget::GetTraceRegionName()
+{
+	return GConfig->GetTextOrDefault(SettingsCategory, TraceRegionSettingName, FText::FromString("ToolbarCustomRegion"), FEditorTraceUtilitiesModule::GetTraceUtilitiesIni());
+}
+
+FText SInsightsStatusBarWidget::GetTraceRegionNameDesc()
+{
+	return RegionIsActive() ? LOCTEXT("TraceRegionNameDisabledDesc", "The name of the region to start or stop. Cannot be edited while a region is active.") : LOCTEXT("TraceRegionNameEnabledDesc", "The name of the region to start or stop.");
+}
+
+void SInsightsStatusBarWidget::ToggleRegion_Execute()
+{
+	if(RegionIsActive())
+	{
+		TRACE_END_REGION(*GetTraceRegionName().ToString());	
+	}
+	else
+	{
+		TRACE_BEGIN_REGION(*GetTraceRegionName().ToString());
+	}
+	bIsRegionActive = !bIsRegionActive;
+}
+
+bool SInsightsStatusBarWidget::RegionIsActive()
+{
+	return bIsRegionActive;
+}
+
+FText SInsightsStatusBarWidget::GetRegionSwitchLabelText()
+{
+	if(RegionIsActive())
+	{
+		return LOCTEXT("TraceEndRegionLabel", "End Region");	
+	}
+	return LOCTEXT("TraceBeginRegionLabel", "Begin Region");
+}
+
+FText SInsightsStatusBarWidget::GetRegionSwitchDescText()
+{
+	if(RegionIsActive())
+	{
+		return LOCTEXT("TraceEndRegionDesc", "Marks the ending of a trace region with the name input above.");	
+	}
+	return LOCTEXT("TraceBeginRegionDesc", "Marks the beginning of a trace region with the name input above.");
+}
+
 void SInsightsStatusBarWidget::PopulateRecentTracesList()
 {
 	Traces.Empty();
diff --git a/Engine/Plugins/TraceUtilities/Source/EditorTraceUtilities/Private/InsightsStatusBar/SInsightsStatusBar.h b/Engine/Plugins/TraceUtilities/Source/EditorTraceUtilities/Private/InsightsStatusBar/SInsightsStatusBar.h
--- a/Engine/Plugins/TraceUtilities/Source/EditorTraceUtilities/Private/InsightsStatusBar/SInsightsStatusBar.h
+++ b/Engine/Plugins/TraceUtilities/Source/EditorTraceUtilities/Private/InsightsStatusBar/SInsightsStatusBar.h
@@ -126,7 +126,17 @@ private:
 
 	bool TraceBookmark_CanExecute();
 	void TraceBookmark_Execute();
-
+	
+	bool bIsRegionActive = false;
+	FText GetTraceRegionName();
+	FText TraceRegionName = FText();
+	FText GetTraceRegionNameDesc();
+
+	void ToggleRegion_Execute();
+	bool RegionIsActive();
+	FText GetRegionSwitchLabelText();
+	FText GetRegionSwitchDescText();
+	
 	void PopulateRecentTracesList();
 
 	void OpenTrace(int32 Index);
@@ -140,6 +150,7 @@ private:
 	static const TCHAR* SettingsCategory;
 	static const TCHAR* OpenLiveSessionOnTraceStartSettingName;
 	static const TCHAR* OpenInsightsAfterTraceSettingName;
+	static const TCHAR* TraceRegionSettingName;
 	static const TCHAR* ShowInExplorerAfterTraceSettingName;
 
 	bool bIsTraceRecordButtonHovered = false;
