Trace destination now reflects the trace target when trace is already running when the session is selected in the Trace Control Window.
Fix UI issues when tracing is disabled in Trace Control Window and Insights Status Bar Widget.

diff --git a/Engine/Plugins/TraceUtilities/Source/EditorTraceUtilities/Private/InsightsStatusBar/SInsightsStatusBar.cpp b/Engine/Plugins/TraceUtilities/Source/EditorTraceUtilities/Private/InsightsStatusBar/SInsightsStatusBar.cpp
--- a/Engine/Plugins/TraceUtilities/Source/EditorTraceUtilities/Private/InsightsStatusBar/SInsightsStatusBar.cpp
+++ b/Engine/Plugins/TraceUtilities/Source/EditorTraceUtilities/Private/InsightsStatusBar/SInsightsStatusBar.cpp
@@ -195,6 +195,7 @@ void SInsightsStatusBarWidget::Construct(const FArguments& InArgs)
 			.HAlign(HAlign_Left)
 			.VAlign(VAlign_Bottom)
 			.OnClicked_Lambda([this]() { SaveSnapshot(); return FReply::Handled(); })
+			.IsEnabled(this, &SInsightsStatusBarWidget::SaveSnapshot_CanExecute)
 			.Content()
 			[
 				SNew(SImage)
@@ -285,7 +286,8 @@ TSharedRef<SWidget> SInsightsStatusBarWidget::MakeTraceMenu()
 			GetRegionSwitchLabelText(),
 			GetRegionSwitchDescText(),
 			FSlateIcon(FEditorTraceUtilitiesStyle::Get().GetStyleSetName(), RegionIsActive() ? "Icons.EndRegion.Menu" : "Icons.BeginRegion.Menu"),
-			FUIAction(FExecuteAction::CreateSP(this, &SInsightsStatusBarWidget::ToggleRegion_Execute)),
+			FUIAction(FExecuteAction::CreateSP(this, &SInsightsStatusBarWidget::ToggleRegion_Execute),
+					  FCanExecuteAction::CreateSP(this, &SInsightsStatusBarWidget::ToggleRegion_CanExecute)),
 			NAME_None,
 			EUserInterfaceActionType::Button
 		);
@@ -335,7 +337,8 @@ TSharedRef<SWidget> SInsightsStatusBarWidget::MakeTraceMenu()
 			TAttribute<FText>::CreateSP(this, &SInsightsStatusBarWidget::GetTraceMenuItemText),
 			TAttribute<FText>::CreateSP(this, &SInsightsStatusBarWidget::GetTraceMenuItemTooltipText),
 			FSlateIcon(FEditorTraceUtilitiesStyle::Get().GetStyleSetName(), "Icons.StartTrace.Menu"),
-			FUIAction(FExecuteAction::CreateSP(this, &SInsightsStatusBarWidget::ToggleTrace_OnClicked)),
+			FUIAction(FExecuteAction::CreateSP(this, &SInsightsStatusBarWidget::ToggleTrace_OnClicked),
+					  FCanExecuteAction::CreateSP(this, &SInsightsStatusBarWidget::ToggleTrace_CanExecute)),
 			NAME_None,
 			EUserInterfaceActionType::Button
 		);
@@ -538,6 +541,7 @@ void SInsightsStatusBarWidget::InitCommandList()
 
 FText SInsightsStatusBarWidget::GetTitleToolTipText() const
 {
+#if UE_TRACE_ENABLED
 	FTextBuilder DescBuilder;
 
 	const FString Dest = FTraceAuxiliary::GetTraceDestinationString();
@@ -556,6 +560,9 @@ FText SInsightsStatusBarWidget::GetTitleToolTipText() const
 	}
 
 	return DescBuilder.ToText();
+#else
+	return LOCTEXT("TraceStatusDisabled", "Trace system is disabled at compile time. Check the UE_TRACE_ENABLED define.");
+#endif
 }
 
 void SInsightsStatusBarWidget::LogMessage(const FText& Text)
@@ -716,9 +723,13 @@ void SInsightsStatusBarWidget::SaveSnapshot()
 	LogMessage(LOCTEXT("SnapshotSavedError", "The snapshot could not be saved."));
 }
 
-bool SInsightsStatusBarWidget::SaveSnapshot_CanExecute()
+bool SInsightsStatusBarWidget::SaveSnapshot_CanExecute() const
 {
+#if UE_TRACE_ENABLED
 	return true;
+#else
+	return false;
+#endif
 }
 
 FText SInsightsStatusBarWidget::GetTraceMenuItemText() const
@@ -733,12 +744,16 @@ FText SInsightsStatusBarWidget::GetTraceMenuItemText() const
 
 FText SInsightsStatusBarWidget::GetTraceMenuItemTooltipText() const
 {
+#if UE_TRACE_ENABLED
 	if (UE::Trace::IsTracing())
 	{
 		return LOCTEXT("StopTraceButtonTooltip", "Stop tracing");
 	}
 
 	return LOCTEXT("StartTraceButtonTooltip", "Start tracing to the selected trace destination.");
+#else
+	return LOCTEXT("StartTraceDisabledButtonTooltip", "Trace system is disabled at compile time. Check the UE_TRACE_ENABLED define.");
+#endif
 }
 
 bool SInsightsStatusBarWidget::ToggleTrace_CanExecute() const
@@ -1045,6 +1060,15 @@ void SInsightsStatusBarWidget::ToggleRegion_Execute()
 	bIsRegionActive = !bIsRegionActive;
 }
 
+bool SInsightsStatusBarWidget::ToggleRegion_CanExecute() const
+{
+#if UE_TRACE_ENABLED
+	return  true;
+#else
+	return false;
+#endif
+}
+
 bool SInsightsStatusBarWidget::RegionIsActive()
 {
 	return bIsRegionActive;
diff --git a/Engine/Plugins/TraceUtilities/Source/EditorTraceUtilities/Private/InsightsStatusBar/SInsightsStatusBar.h b/Engine/Plugins/TraceUtilities/Source/EditorTraceUtilities/Private/InsightsStatusBar/SInsightsStatusBar.h
--- a/Engine/Plugins/TraceUtilities/Source/EditorTraceUtilities/Private/InsightsStatusBar/SInsightsStatusBar.h
+++ b/Engine/Plugins/TraceUtilities/Source/EditorTraceUtilities/Private/InsightsStatusBar/SInsightsStatusBar.h
@@ -83,7 +83,7 @@ private:
 	bool SetTraceDestination_IsChecked(ETraceDestination InDestination);
 
 	void SaveSnapshot();
-	bool SaveSnapshot_CanExecute();
+	bool SaveSnapshot_CanExecute() const;
 
 	FText GetTraceMenuItemText() const;
 	FText GetTraceMenuItemTooltipText() const;
@@ -133,6 +133,7 @@ private:
 	FText GetTraceRegionNameDesc();
 
 	void ToggleRegion_Execute();
+	bool ToggleRegion_CanExecute() const;
 	bool RegionIsActive();
 	FText GetRegionSwitchLabelText();
 	FText GetRegionSwitchDescText();
diff --git a/Engine/Source/Developer/TraceTools/Private/Widgets/STraceControlToolbar.cpp b/Engine/Source/Developer/TraceTools/Private/Widgets/STraceControlToolbar.cpp
--- a/Engine/Source/Developer/TraceTools/Private/Widgets/STraceControlToolbar.cpp
+++ b/Engine/Source/Developer/TraceTools/Private/Widgets/STraceControlToolbar.cpp
@@ -196,7 +196,7 @@ void STraceControlToolbar::SetTraceTarget_Execute(ETraceTarget InTraceTarget)
 
 bool STraceControlToolbar::ToggleTrace_CanExecute() const
 {
-	return TraceController->HasAvailableSelectedInstance();
+	return TraceController->HasAvailableSelectedInstance() && bIsTracingAvailable;
 }
 
 void STraceControlToolbar::ToggleTrace_Execute()
@@ -228,7 +228,7 @@ void STraceControlToolbar::ToggleTrace_Execute()
 
 bool STraceControlToolbar::TraceSnapshot_CanExecute() const
 {
-	return TraceController->HasAvailableSelectedInstance();
+	return TraceController->HasAvailableSelectedInstance() && bIsTracingAvailable;
 }
 
 void STraceControlToolbar::TraceSnapshot_Execute()
@@ -248,7 +248,7 @@ void STraceControlToolbar::TraceSnapshot_Execute()
 
 bool STraceControlToolbar::TraceBookmark_CanExecute() const
 {
-	return TraceController->HasAvailableSelectedInstance() && bIsTracing && !bIsPaused;
+	return TraceController->HasAvailableSelectedInstance() && bIsTracingAvailable && bIsTracing && !bIsPaused;
 }
 
 void STraceControlToolbar::TraceBookmark_Execute()
@@ -322,6 +322,16 @@ void STraceControlToolbar::OnTraceStatusUpdated(const FTraceStatus& InStatus, FT
 	bIsTracing = InStatus.bIsTracing;
 	bIsPaused = InStatus.bIsPaused;
 	bAreStatNamedEventsEnabled = InStatus.bAreStatNamedEventsEnabled;
+	bIsTracingAvailable = InStatus.TraceSystemStatus != FTraceStatus::ETraceSystemStatus::NotAvailable;
+
+	if (InStatus.TraceSystemStatus == FTraceStatus::ETraceSystemStatus::TracingToServer)
+	{
+		TraceTarget = ETraceTarget::Server;
+	}
+	else if (InStatus.TraceSystemStatus == FTraceStatus::ETraceSystemStatus::TracingToFile)
+	{
+		TraceTarget = ETraceTarget::File;
+	}
 }
 
 bool STraceControlToolbar::TogglePauseResume_CanExecute() const
diff --git a/Engine/Source/Developer/TraceTools/Private/Widgets/STraceControlToolbar.h b/Engine/Source/Developer/TraceTools/Private/Widgets/STraceControlToolbar.h
--- a/Engine/Source/Developer/TraceTools/Private/Widgets/STraceControlToolbar.h
+++ b/Engine/Source/Developer/TraceTools/Private/Widgets/STraceControlToolbar.h
@@ -83,6 +83,7 @@ private:
 	bool bIsTracing = false;
 	bool bIsPaused = false;
 	bool bAreStatNamedEventsEnabled = false;
+	bool bIsTracingAvailable = false;
 	FString TraceHostAddr;
 
 	FDelegateHandle OnStatusReceivedDelegate;
diff --git a/Engine/Source/Developer/TraceTools/Private/Widgets/STraceStatistics.cpp b/Engine/Source/Developer/TraceTools/Private/Widgets/STraceStatistics.cpp
--- a/Engine/Source/Developer/TraceTools/Private/Widgets/STraceStatistics.cpp
+++ b/Engine/Source/Developer/TraceTools/Private/Widgets/STraceStatistics.cpp
@@ -426,7 +426,7 @@ FText STraceStatistics::GetStatsCacheText() const
 {
 	if (!SessionFilterService->HasStats())
 	{
-		return FText();
+		return LOCTEXT("N/A", "N/A");
 	}
 
 	FNumberFormattingOptions FormattingOptionsMem;
diff --git a/Engine/Source/Runtime/TraceLog/Public/Trace/Trace.h b/Engine/Source/Runtime/TraceLog/Public/Trace/Trace.h
--- a/Engine/Source/Runtime/TraceLog/Public/Trace/Trace.h
+++ b/Engine/Source/Runtime/TraceLog/Public/Trace/Trace.h
@@ -149,12 +149,12 @@ typedef bool		ChannelIterCallback(const FChannelInfo& Info, void*/*User*/);
 
 struct FStatistics
 {
-	uint64		BytesSent;
-	uint64		BytesTraced;
-	uint64		MemoryUsed;
-	uint32		CacheAllocated;	// Total memory allocated in cache buffers
-	uint32		CacheUsed;		// Used cache memory; Important-marked events are stored in the cache.
-	uint32		CacheWaste;		// Unused memory from retired cache buffers
+	uint64		BytesSent		= 0;
+	uint64		BytesTraced		= 0;
+	uint64		MemoryUsed		= 0;
+	uint32		CacheAllocated	= 0;	// Total memory allocated in cache buffers
+	uint32		CacheUsed		= 0;	// Used cache memory; Important-marked events are stored in the cache.
+	uint32		CacheWaste		= 0;	// Unused memory from retired cache buffers
 };
 
 struct FSendFlags
