Add more descriptive tooltips to the channels and to the Screenshot,Bookmark and Region menu entries in the Insights Status Bar Widget.

diff --git a/Engine/Plugins/TraceUtilities/Source/EditorTraceUtilities/Private/InsightsStatusBar/SInsightsStatusBar.cpp b/Engine/Plugins/TraceUtilities/Source/EditorTraceUtilities/Private/InsightsStatusBar/SInsightsStatusBar.cpp
--- a/Engine/Plugins/TraceUtilities/Source/EditorTraceUtilities/Private/InsightsStatusBar/SInsightsStatusBar.cpp
+++ b/Engine/Plugins/TraceUtilities/Source/EditorTraceUtilities/Private/InsightsStatusBar/SInsightsStatusBar.cpp
@@ -258,7 +258,7 @@ TSharedRef<SWidget> SInsightsStatusBarWidget::MakeTraceMenu()
 			FInsightsStatusBarWidgetCommands::Get().Command_TraceScreenshot,
 			NAME_None,
 			TAttribute<FText>(),
-			TAttribute<FText>(),
+			TAttribute<FText>::CreateSP(this, &SInsightsStatusBarWidget::GetTraceScreenshotTooltipText),
 			FSlateIcon(FEditorTraceUtilitiesStyle::Get().GetStyleSetName(), "Icons.Screenshot.Menu")
 		);
 
@@ -267,7 +267,7 @@ TSharedRef<SWidget> SInsightsStatusBarWidget::MakeTraceMenu()
 			FInsightsStatusBarWidgetCommands::Get().Command_TraceBookmark,
 			NAME_None,
 			TAttribute<FText>(),
-			TAttribute<FText>(),
+			TAttribute<FText>::CreateSP(this, &SInsightsStatusBarWidget::GetTraceBookmarkTooltipText),
 			FSlateIcon(FEditorTraceUtilitiesStyle::Get().GetStyleSetName(), "Icons.Bookmark.Menu")
 		);
 		
@@ -284,7 +284,7 @@ TSharedRef<SWidget> SInsightsStatusBarWidget::MakeTraceMenu()
 		
 		MenuBuilder.AddMenuEntry(
 			GetRegionSwitchLabelText(),
-			GetRegionSwitchDescText(),
+			TAttribute<FText>::CreateSP(this, &SInsightsStatusBarWidget::GetRegionSwitchDescText),
 			FSlateIcon(FEditorTraceUtilitiesStyle::Get().GetStyleSetName(), RegionIsActive() ? "Icons.EndRegion.Menu" : "Icons.BeginRegion.Menu"),
 			FUIAction(FExecuteAction::CreateSP(this, &SInsightsStatusBarWidget::ToggleRegion_Execute),
 					  FCanExecuteAction::CreateSP(this, &SInsightsStatusBarWidget::ToggleRegion_CanExecute)),
@@ -499,10 +499,22 @@ void SInsightsStatusBarWidget::Channels_BuildMenu(FMenuBuilder& MenuBuilder)
 			const FChannelData& Data = ChannelsInfo[Index];
 			FString ChannelDisplayName = Data.Name;
 			ChannelDisplayName.RemoveFromEnd(TEXT("Channel"), 7);
+			FText ChannelTooltip = FText::FromString(Data.Desc);
+			if (Data.bIsReadOnly)
+			{
+				if (Data.Desc.EndsWith("."))
+				{
+					ChannelTooltip = FText::Format(LOCTEXT("ChannelTooltipFmt1", "{0} This channel is readonly and can only be enabled from the command line."), FText::FromString(Data.Desc));
+				}
+				else
+				{
+					ChannelTooltip = FText::Format(LOCTEXT("ChannelTooltipFmt2", "{0}. This channel is readonly and can only be enabled from the command line."), FText::FromString(Data.Desc));
+				}
+			}
 
 			MenuBuilder.AddMenuEntry(
 				FText::FromString(ChannelDisplayName),
-				FText::Format(LOCTEXT("ChannelDesc", "Enable/disable the {0} channel"), FText::FromString(ChannelDisplayName)),
+				ChannelTooltip,
 				FSlateIcon(),
 				FUIAction(FExecuteAction::CreateSP(this, &SInsightsStatusBarWidget::ToggleChannel_Execute, Index),
 					FCanExecuteAction::CreateLambda([Value = !Data.bIsReadOnly]() { return Value; }),
@@ -944,6 +956,7 @@ void SInsightsStatusBarWidget::CreateChannelsInfo()
 
 		FChannelData NewChannelInfo;
 		NewChannelInfo.Name = Info.Name;
+		NewChannelInfo.Desc = Info.Desc;
 		NewChannelInfo.bIsEnabled = Info.bIsEnabled;
 		NewChannelInfo.bIsReadOnly = Info.bIsReadOnly;
 
@@ -1014,7 +1027,7 @@ bool SInsightsStatusBarWidget::ToggleChannel_IsChecked(int32 Index)
 }
 
 
-bool SInsightsStatusBarWidget::TraceScreenshot_CanExecute()
+bool SInsightsStatusBarWidget::TraceScreenshot_CanExecute() const
 {
 	return SHOULD_TRACE_SCREENSHOT();
 }
@@ -1026,7 +1039,7 @@ void SInsightsStatusBarWidget::TraceScreenshot_Execute()
 #endif
 }
 
-bool SInsightsStatusBarWidget::TraceBookmark_CanExecute()
+bool SInsightsStatusBarWidget::TraceBookmark_CanExecute() const
 {
 	return SHOULD_TRACE_BOOKMARK();
 }
@@ -1051,30 +1064,26 @@ void SInsightsStatusBarWidget::ToggleRegion_Execute()
 {
 	if(RegionIsActive())
 	{
-		TRACE_END_REGION(*GetTraceRegionName().ToString());	
+		TRACE_END_REGION_WITH_ID(RegionId);
+		RegionId = 0;
 	}
 	else
 	{
-		TRACE_BEGIN_REGION(*GetTraceRegionName().ToString());
+		RegionId = TRACE_BEGIN_REGION_WITH_ID(*GetTraceRegionName().ToString());
 	}
-	bIsRegionActive = !bIsRegionActive;
 }
 
 bool SInsightsStatusBarWidget::ToggleRegion_CanExecute() const
 {
-#if UE_TRACE_ENABLED
-	return  true;
-#else
-	return false;
-#endif
+	return SHOULD_TRACE_REGION();
 }
 
-bool SInsightsStatusBarWidget::RegionIsActive()
+bool SInsightsStatusBarWidget::RegionIsActive() const
 {
-	return bIsRegionActive;
+	return RegionId > 0;
 }
 
-FText SInsightsStatusBarWidget::GetRegionSwitchLabelText()
+FText SInsightsStatusBarWidget::GetRegionSwitchLabelText() const
 {
 	if(RegionIsActive())
 	{
@@ -1083,12 +1092,18 @@ FText SInsightsStatusBarWidget::GetRegionSwitchLabelText()
 	return LOCTEXT("TraceBeginRegionLabel", "Begin Region");
 }
 
-FText SInsightsStatusBarWidget::GetRegionSwitchDescText()
+FText SInsightsStatusBarWidget::GetRegionSwitchDescText() const
 {
+	if (!ToggleRegion_CanExecute())
+	{
+		LOCTEXT("TraceRegionDisabledDesc", "Regions can only be traced when the Region channel in enabled.");
+	}
+
 	if(RegionIsActive())
 	{
 		return LOCTEXT("TraceEndRegionDesc", "Marks the ending of a trace region with the name input above.");	
 	}
+
 	return LOCTEXT("TraceBeginRegionDesc", "Marks the beginning of a trace region with the name input above.");
 }
 
@@ -1149,6 +1164,30 @@ void SInsightsStatusBarWidget::OpenTrace(int32 Index)
 	}
 }
 
+FText SInsightsStatusBarWidget::GetTraceScreenshotTooltipText() const
+{
+	if (TraceScreenshot_CanExecute())
+	{
+		return LOCTEXT("TraceScreenshotTooltip1", "Takes a screenshot and sends it to the trace.");
+	}
+	else
+	{
+		return LOCTEXT("TraceScreenshotTooltip2", "Screenshots can only be traced when the Screenshot channel is enabled.");
+	}
+}
+
+FText SInsightsStatusBarWidget::GetTraceBookmarkTooltipText() const
+{
+	if (TraceBookmark_CanExecute())
+	{
+		return LOCTEXT("TraceBookmarkTooltip1", "Traces a bookmark.");
+	}
+	else
+	{
+		return LOCTEXT("TraceBookmarkTooltip2", "Bookmarks can only be traced when the Bookmark channel is enabled.");
+	}
+}
+
 } // namespace UE::EditorTraceUtilities
 
 #undef LOCTEXT_NAMESPACE
diff --git a/Engine/Plugins/TraceUtilities/Source/EditorTraceUtilities/Private/InsightsStatusBar/SInsightsStatusBar.h b/Engine/Plugins/TraceUtilities/Source/EditorTraceUtilities/Private/InsightsStatusBar/SInsightsStatusBar.h
--- a/Engine/Plugins/TraceUtilities/Source/EditorTraceUtilities/Private/InsightsStatusBar/SInsightsStatusBar.h
+++ b/Engine/Plugins/TraceUtilities/Source/EditorTraceUtilities/Private/InsightsStatusBar/SInsightsStatusBar.h
@@ -38,6 +38,7 @@ class SInsightsStatusBarWidget : public SCompoundWidget
 	struct FChannelData
 	{
 		FString Name;
+		FString Desc;
 		bool bIsEnabled = false;
 		bool bIsReadOnly = false;
 	};
@@ -121,22 +122,23 @@ private:
 
 	void InitCommandList();
 
-	bool TraceScreenshot_CanExecute();
+	bool TraceScreenshot_CanExecute() const;
 	void TraceScreenshot_Execute();
+	FText GetTraceScreenshotTooltipText() const;
 
-	bool TraceBookmark_CanExecute();
+	bool TraceBookmark_CanExecute() const;
 	void TraceBookmark_Execute();
+	FText GetTraceBookmarkTooltipText() const;
 	
-	bool bIsRegionActive = false;
 	FText GetTraceRegionName();
 	FText TraceRegionName = FText();
 	FText GetTraceRegionNameDesc();
 
 	void ToggleRegion_Execute();
 	bool ToggleRegion_CanExecute() const;
-	bool RegionIsActive();
-	FText GetRegionSwitchLabelText();
-	FText GetRegionSwitchDescText();
+	bool RegionIsActive() const;
+	FText GetRegionSwitchLabelText() const;
+	FText GetRegionSwitchDescText() const;
 	
 	void PopulateRecentTracesList();
 
@@ -170,5 +172,7 @@ private:
 
 	TArray<TSharedPtr<FTraceFileInfo>> Traces;
 	FName LogListingName;
+
+	uint64 RegionId = 0;
 };
 }
diff --git a/Engine/Source/Developer/TraceTools/Private/Models/TraceChannel.cpp b/Engine/Source/Developer/TraceTools/Private/Models/TraceChannel.cpp
--- a/Engine/Source/Developer/TraceTools/Private/Models/TraceChannel.cpp
+++ b/Engine/Source/Developer/TraceTools/Private/Models/TraceChannel.cpp
@@ -2,11 +2,14 @@
 
 #include "TraceChannel.h"
 
+#include "Internationalization/Internationalization.h"
 #include "Internationalization/Text.h"
 
 // TraceTools
 #include "Services/ISessionTraceFilterService.h"
 
+#define LOCTEXT_NAMESPACE "UE::TraceTools::FTraceChannel"
+
 namespace UE::TraceTools
 {
 
@@ -35,7 +38,20 @@ FText FTraceChannel::GetDisplayText() const
 
 FText FTraceChannel::GetTooltipText() const
 {
-	return FText::FromString(Description);
+	FText ChannelTooltip = FText::FromString(Description);
+	if (bReadOnly)
+	{
+		if (Description.EndsWith("."))
+		{
+			ChannelTooltip = FText::Format(LOCTEXT("ChannelTooltipFmt1", "{0} This channel is readonly and can only be enabled from the command line."), FText::FromString(Description));
+		}
+		else
+		{
+			ChannelTooltip = FText::Format(LOCTEXT("ChannelTooltipFmt2", "{0}. This channel is readonly and can only be enabled from the command line."), FText::FromString(Description));
+		}
+	}
+
+	return ChannelTooltip;
 }
 
 FString FTraceChannel::GetName() const
@@ -79,4 +95,6 @@ void FTraceChannel::GetSearchString(TArray<FString>& OutFilterStrings) const
 	OutFilterStrings.Add(Name);
 }
 
-} // namespace UE::TraceTools
\ No newline at end of file
+} // namespace UE::TraceTools
+
+#undef LOCTEXT_NAMESPACE // UE::TraceTools::FTraceChannel
diff --git a/Engine/Source/Runtime/Core/Private/ProfilingDebugging/MiscTrace.cpp b/Engine/Source/Runtime/Core/Private/ProfilingDebugging/MiscTrace.cpp
--- a/Engine/Source/Runtime/Core/Private/ProfilingDebugging/MiscTrace.cpp
+++ b/Engine/Source/Runtime/Core/Private/ProfilingDebugging/MiscTrace.cpp
@@ -203,4 +203,9 @@ bool FMiscTrace::ShouldTraceBookmark()
 {
 	return UE_TRACE_CHANNELEXPR_IS_ENABLED(BookmarkChannel);
 }
+
+bool FMiscTrace::ShouldTraceRegion()
+{
+	return UE_TRACE_CHANNELEXPR_IS_ENABLED(RegionChannel);
+}
 #endif
diff --git a/Engine/Source/Runtime/Core/Public/ProfilingDebugging/MiscTrace.h b/Engine/Source/Runtime/Core/Public/ProfilingDebugging/MiscTrace.h
--- a/Engine/Source/Runtime/Core/Public/ProfilingDebugging/MiscTrace.h
+++ b/Engine/Source/Runtime/Core/Public/ProfilingDebugging/MiscTrace.h
@@ -84,6 +84,7 @@ struct FMiscTrace
 	CORE_API static void OutputScreenshot(const TCHAR* Name, uint64 Cycle, uint32 Width, uint32 Height, TArray64<uint8> Data);
 	CORE_API static bool ShouldTraceScreenshot();
 	CORE_API static bool ShouldTraceBookmark();
+	CORE_API static bool ShouldTraceRegion();
 
 private:
 	CORE_API static void OutputBookmarkInternal(const void* BookmarkPoint, uint16 EncodedFormatArgsSize, uint8* EncodedFormatArgs);
@@ -141,6 +142,9 @@ private:
 #define SHOULD_TRACE_BOOKMARK() \
 	FMiscTrace::ShouldTraceBookmark()
 
+#define SHOULD_TRACE_REGION() \
+	FMiscTrace::ShouldTraceRegion()
+
 #else
 
 #define TRACE_BOOKMARK(...)
@@ -154,5 +158,6 @@ private:
 #define TRACE_SCREENSHOT(...)
 #define SHOULD_TRACE_SCREENSHOT(...) false
 #define SHOULD_TRACE_BOOKMARK(...) false
+#define SHOULD_TRACE_REGION(...) false
 
 #endif
