Fix the Open Live Session on Trace Start functionality in the trace status bar widget.

diff --git a/Engine/Plugins/TraceUtilities/Source/EditorTraceUtilities/Private/InsightsStatusBar/SInsightsStatusBar.cpp b/Engine/Plugins/TraceUtilities/Source/EditorTraceUtilities/Private/InsightsStatusBar/SInsightsStatusBar.cpp
--- a/Engine/Plugins/TraceUtilities/Source/EditorTraceUtilities/Private/InsightsStatusBar/SInsightsStatusBar.cpp
+++ b/Engine/Plugins/TraceUtilities/Source/EditorTraceUtilities/Private/InsightsStatusBar/SInsightsStatusBar.cpp
@@ -18,6 +18,7 @@
 #include "Misc/DateTime.h"
 #include "Misc/Paths.h"
 #include "Modules/ModuleManager.h"
+#include "ProfilingDebugging/CpuProfilerTrace.h"
 #include "ProfilingDebugging/MiscTrace.h"
 #include "ProfilingDebugging/PlatformEvents.h"
 #include "ProfilingDebugging/TraceScreenshot.h"
@@ -63,22 +64,6 @@ const TCHAR* SInsightsStatusBarWidget::OpenInsightsAfterTraceSettingName = TEXT(
 const TCHAR* SInsightsStatusBarWidget::TraceRegionSettingName = TEXT("InsightsToolbarTraceRegion");
 const TCHAR* SInsightsStatusBarWidget::ShowInExplorerAfterTraceSettingName = TEXT("ShowInExplorerAfterTrace");
 
-class FOpenLiveSessionTask
-{
-public:
-	FOpenLiveSessionTask()
-	{}
-
-	FORCEINLINE TStatId GetStatId() const { RETURN_QUICK_DECLARE_CYCLE_STAT(FOpenLiveSessionTask, STATGROUP_TaskGraphTasks); }
-	ENamedThreads::Type GetDesiredThread() { return ENamedThreads::Type::GameThread; }
-	static ESubsequentsMode::Type GetSubsequentsMode() { return ESubsequentsMode::TrackSubsequents; }
-
-	void DoTask(ENamedThreads::Type CurrentThread, const FGraphEventRef& MyCompletionGraphEvent)
-	{
-		FUnrealInsightsLauncher::Get()->TryOpenTraceFromDestination(FTraceAuxiliary::GetTraceDestinationString());
-	}
-};
-
 ////////////////////////////////////////////////////////////////////////////////////////////////////
 // FInsightsStatusBarWidgetCommands
 ////////////////////////////////////////////////////////////////////////////////////////////////////
@@ -130,9 +115,24 @@ void RegisterInsightsStatusWidgetWithToolMenu()
 
 void SInsightsStatusBarWidget::Tick(const FGeometry& AllottedGeometry, const double InCurrentTime, const float InDeltaTime)
 {
+	TRACE_CPUPROFILER_EVENT_SCOPE(SInsightsStatusBarWidget::Tick);
+
 	SCompoundWidget::Tick(AllottedGeometry, InCurrentTime, InDeltaTime);
 
 	bShouldUpdateChannels = true;
+
+	if (bShouldTryOpenLiveSession)
+	{
+		constexpr double Delay = 1.0f;
+		OpenLiveSessionScheduledTime = InCurrentTime + Delay;
+		bShouldTryOpenLiveSession = false;
+	}
+
+	if (OpenLiveSessionScheduledTime > 0.0f && OpenLiveSessionScheduledTime <= InCurrentTime)
+	{
+		OpenLiveSessionScheduledTime = 0.0f;
+		FUnrealInsightsLauncher::Get()->TryOpenTraceFromDestination(FTraceAuxiliary::GetTraceDestinationString());
+	}
 }
 
 void SInsightsStatusBarWidget::Construct(const FArguments& InArgs)
@@ -871,7 +871,7 @@ void SInsightsStatusBarWidget::OnTraceStarted(FTraceAuxiliary::EConnectionType I
 {
 	if (InTraceType == FTraceAuxiliary::EConnectionType::Network && GetBooleanSettingValue(OpenLiveSessionOnTraceStartSettingName))
 	{
-		TGraphTask<FOpenLiveSessionTask>::CreateTask().ConstructAndDispatchWhenReady();
+		bShouldTryOpenLiveSession = true;
 	}
 }
 
diff --git a/Engine/Plugins/TraceUtilities/Source/EditorTraceUtilities/Private/InsightsStatusBar/SInsightsStatusBar.h b/Engine/Plugins/TraceUtilities/Source/EditorTraceUtilities/Private/InsightsStatusBar/SInsightsStatusBar.h
--- a/Engine/Plugins/TraceUtilities/Source/EditorTraceUtilities/Private/InsightsStatusBar/SInsightsStatusBar.h
+++ b/Engine/Plugins/TraceUtilities/Source/EditorTraceUtilities/Private/InsightsStatusBar/SInsightsStatusBar.h
@@ -174,5 +174,8 @@ private:
 	FName LogListingName;
 
 	uint64 RegionId = 0;
+
+	bool bShouldTryOpenLiveSession = false;
+	double OpenLiveSessionScheduledTime = 0.0f;
 };
 }
diff --git a/Engine/Plugins/TraceUtilities/Source/EditorTraceUtilities/Private/UnrealInsightsLauncher.cpp b/Engine/Plugins/TraceUtilities/Source/EditorTraceUtilities/Private/UnrealInsightsLauncher.cpp
--- a/Engine/Plugins/TraceUtilities/Source/EditorTraceUtilities/Private/UnrealInsightsLauncher.cpp
+++ b/Engine/Plugins/TraceUtilities/Source/EditorTraceUtilities/Private/UnrealInsightsLauncher.cpp
@@ -10,6 +10,7 @@
 #include "Logging/MessageLog.h"
 #include "MessageLogModule.h"
 #include "ToolMenus.h"
+#include "ProfilingDebugging/CpuProfilerTrace.h"
 #include "ProfilingDebugging/TraceAuxiliary.h"
 #include "Trace/StoreClient.h"
 
@@ -220,12 +221,19 @@ void FUnrealInsightsLauncher::LogMessageOnGameThread(const FText& Message)
 
 bool FUnrealInsightsLauncher::TryOpenTraceFromDestination(const FString& Destination)
 {
+	TRACE_CPUPROFILER_EVENT_SCOPE(TryOpenTraceFromDestination);
+
 	if (Destination.IsEmpty())
 		return false;
 	
-	if (FPaths::FileExists(Destination) && FPaths::GetExtension(Destination).Equals(TEXT("utrace")))
+	if (FPaths::GetExtension(Destination).Equals(TEXT("utrace")))
 	{
-		return OpenTraceFile(Destination);
+		if (FPaths::FileExists(Destination))
+		{
+			return OpenTraceFile(Destination);
+		}
+		
+		return false;
 	}
 
 	// assume it's a ip address of the target server for now so we don't introduce dependencies on the socket system
@@ -267,6 +275,8 @@ bool FUnrealInsightsLauncher::OpenRemoteTrace(const FString& TraceHostAddress, c
 
 bool FUnrealInsightsLauncher::OpenActiveTraceFromStore(const FString& TraceHostAddress)
 {
+	TRACE_CPUPROFILER_EVENT_SCOPE(OpenActiveTraceFromStore);
+
 	UE::Trace::FStoreClient* StoreClient = UE::Trace::FStoreClient::Connect(*TraceHostAddress);
 
 	if (!StoreClient)
@@ -317,9 +327,12 @@ bool FUnrealInsightsLauncher::OpenActiveTraceFromStore(const FString& TraceHostA
 
 	uint32 TraceId = SessionInfo->GetTraceId();
 
-	if (!OpenRemoteTrace(TraceHostAddress, (uint16)StoreClient->GetStorePort(), TraceId))
+	if (TraceId == 0 || !OpenRemoteTrace(TraceHostAddress, (uint16)StoreClient->GetStorePort(), TraceId))
 	{
-		//Failed to open Trace File %s
+		const FText MessageBoxTextFmt = LOCTEXT("StoreClientFailedToOpenRemote_TextFmt",
+			"Failed to open remote session from address {0} with TraceId: {1}");
+		const FText MessageBoxText = FText::Format(MessageBoxTextFmt, FText::FromString(TraceHostAddress), FText::AsNumber(TraceId));
+		LogMessage(MessageBoxText);
 		return false;
 	}
 
