Small fixes to the temporary timing view workflow for starting analysis in the Editor. Fixes some static analysis issues.

diff --git a/Engine/Plugins/TraceUtilities/Source/InsightsEditor/Private/InsightsEditorModule.cpp b/Engine/Plugins/TraceUtilities/Source/InsightsEditor/Private/InsightsEditorModule.cpp
--- a/Engine/Plugins/TraceUtilities/Source/InsightsEditor/Private/InsightsEditorModule.cpp
+++ b/Engine/Plugins/TraceUtilities/Source/InsightsEditor/Private/InsightsEditorModule.cpp
@@ -91,7 +91,7 @@ void FInsightsEditorModule::StartupModule()
 		{
 			if (InMajorTabId != FInsightsManagerTabs::MemoryProfilerTabId)
 			{
-				StartTrace();
+				StartTraceAnalysis();
 			}
 		});
 
@@ -114,18 +114,13 @@ void FInsightsEditorModule::ShutdownModule()
 	LLM_SCOPE_BYNAME(TEXT("Insights/InsightsEditor"));
 }
 
-void FInsightsEditorModule::StartTrace()
+void FInsightsEditorModule::StartTraceAnalysis()
 {
 	LLM_SCOPE_BYNAME(TEXT("Insights/InsightsEditor"));
 
-	if (!bTraceStarted)
+	IUnrealInsightsModule& UnrealInsightsModule = FModuleManager::LoadModuleChecked<IUnrealInsightsModule>("TraceInsights");
+	if (!UnrealInsightsModule.GetAnalysisSession().IsValid())
 	{
-		//bTraceStarted = FTraceAuxiliary::Start(
-		//	FTraceAuxiliary::EConnectionType::Network,
-		//	TEXT("127.0.0.1"),
-		//	nullptr);
-
-		IUnrealInsightsModule& UnrealInsightsModule = FModuleManager::LoadModuleChecked<IUnrealInsightsModule>("TraceInsights");
 		UnrealInsightsModule.StartAnalysisForLastLiveSession();
 	}
 }
diff --git a/Engine/Plugins/TraceUtilities/Source/InsightsEditor/Private/InsightsEditorModule.h b/Engine/Plugins/TraceUtilities/Source/InsightsEditor/Private/InsightsEditorModule.h
--- a/Engine/Plugins/TraceUtilities/Source/InsightsEditor/Private/InsightsEditorModule.h
+++ b/Engine/Plugins/TraceUtilities/Source/InsightsEditor/Private/InsightsEditorModule.h
@@ -16,10 +16,9 @@ public:
 	//~ Begin IInsightsEditorModule Interface
 	virtual void StartupModule() override;
 	virtual void ShutdownModule() override;
-	virtual void StartTrace() override;
 	//~ End IInsightsEditorModule Interface
 
 private:
-	bool bTraceStarted = false;
+	void StartTraceAnalysis();
 };
 } // namespace UE::InsightsEditor
\ No newline at end of file
diff --git a/Engine/Plugins/TraceUtilities/Source/InsightsEditor/Public/IInsightsEditorModule.h b/Engine/Plugins/TraceUtilities/Source/InsightsEditor/Public/IInsightsEditorModule.h
--- a/Engine/Plugins/TraceUtilities/Source/InsightsEditor/Public/IInsightsEditorModule.h
+++ b/Engine/Plugins/TraceUtilities/Source/InsightsEditor/Public/IInsightsEditorModule.h
@@ -33,7 +33,5 @@ public:
 		static const FName ModuleName = "InsightsEditor";
 		return FModuleManager::Get().IsModuleLoaded(ModuleName);
 	}
-
-	virtual void StartTrace() = 0;
 };
 } // namespace UE::InsightsEditor
\ No newline at end of file
