Removed trace started notification from the Editor that affected  the trace data.
#Synced-CL 47092348

diff --git a/Engine/Plugins/TraceUtilities/Source/EditorTraceUtilities/Private/InsightsStatusBar/SInsightsStatusBar.cpp b/Engine/Plugins/TraceUtilities/Source/EditorTraceUtilities/Private/InsightsStatusBar/SInsightsStatusBar.cpp
--- a/Engine/Plugins/TraceUtilities/Source/EditorTraceUtilities/Private/InsightsStatusBar/SInsightsStatusBar.cpp
+++ b/Engine/Plugins/TraceUtilities/Source/EditorTraceUtilities/Private/InsightsStatusBar/SInsightsStatusBar.cpp
@@ -781,10 +781,16 @@ void SInsightsStatusBarWidget::ToggleTrace_OnClicked()
 {
 	if (UE::Trace::IsTracing())
 	{
+		FString TraceDestinationStr = FTraceAuxiliary::GetTraceDestinationString();
 		bool bResult = FTraceAuxiliary::Stop();
-		if (!bResult)
+
+		if (bResult)
+		{
+			UE_LOG(LogTraceUtilities, Log, TEXT("Stopped tracing to %s."), *TraceDestinationStr);
+		}
+		else
 		{
-			LogMessage(LOCTEXT("TraceStopFailedMsg", "There was no trace connection to stop."));
+			LogMessage(LOCTEXT("TraceStopFailedMsg", "Failed to stop tracing!"));
 		}
 	}
 	else
@@ -798,11 +804,11 @@ void SInsightsStatusBarWidget::ToggleTrace_OnClicked()
 				TraceDestinationStr = TEXT("External Target");
 			}
 
-			ShowNotification(LOCTEXT("TraceMsg", "Trace Started"), FText::Format(LOCTEXT("TracingStartedText", "Trace is now active and saving to the following location (file or tracestore):\n{0}"), FText::FromString(TraceDestinationStr)));
+			UE_LOG(LogTraceUtilities, Log, TEXT("Trace started. Trace is now active and saving to the following location (file or tracestore): %s"), *TraceDestinationStr);
 		}
 		else
 		{
-			LogMessage(LOCTEXT("TraceFailedToStartMsg", "Trace Failed to Start."));
+			LogMessage(LOCTEXT("TraceFailedToStartMsg", "Trace Failed to Start!"));
 		}
 	}
 }
