[Insights] Editor Trace/Insights StatusBar Widget:
* Fixed issue with command actions not being globally accessible, so assigning global key shortcuts to them was not possible.
* Also added a new command for toggling trace recording (so a key shortcut can be assigned to it for ease of use).
* Switched trace functionality by UE_TRACE_MINIMAL_ENABLED (so it could be enabled in shipping).
* Fixed the "recent traces" list to correctly trim older ones (to show max 15 traces).

#Synced-CL 52002978

diff --git a/Engine/Plugins/TraceUtilities/Source/EditorTraceUtilities/Private/InsightsStatusBar/SInsightsStatusBar.cpp b/Engine/Plugins/TraceUtilities/Source/EditorTraceUtilities/Private/InsightsStatusBar/SInsightsStatusBar.cpp
--- a/Engine/Plugins/TraceUtilities/Source/EditorTraceUtilities/Private/InsightsStatusBar/SInsightsStatusBar.cpp
+++ b/Engine/Plugins/TraceUtilities/Source/EditorTraceUtilities/Private/InsightsStatusBar/SInsightsStatusBar.cpp
@@ -10,6 +10,7 @@
 #include "Framework/Notifications/NotificationManager.h"
 #include "GenericPlatform/GenericPlatformFile.h"
 #include "HAL/FileManager.h"
+#include "Interfaces/IMainFrameModule.h"
 #include "Logging/MessageLog.h"
 #include "Logging/TokenizedMessage.h"
 #include "Math/Color.h"
@@ -54,7 +55,7 @@ namespace UE::EditorTraceUtilities
 FStatusBarTraceSettings SInsightsStatusBarWidget::StatusBarTraceSettings;
 
 const TCHAR* SInsightsStatusBarWidget::DefaultPreset = TEXT("default");
-const TCHAR* SInsightsStatusBarWidget::MemoryPreset= TEXT("default,memory");
+const TCHAR* SInsightsStatusBarWidget::MemoryPreset = TEXT("default,memory");
 const TCHAR* SInsightsStatusBarWidget::TaskGraphPreset = TEXT("default,task");
 const TCHAR* SInsightsStatusBarWidget::ContextSwitchesPreset = TEXT("default,contextswitches");
 
@@ -90,13 +91,17 @@ public:
 	{
 		UI_COMMAND(Command_TraceScreenshot, "Trace Screenshot", "Takes a screenshot and sends it to the trace.", EUserInterfaceActionType::Button, FInputChord(EModifierKey::Control, EKeys::F9));
 		UI_COMMAND(Command_TraceBookmark, "Trace Bookmark", "Traces a bookmark.", EUserInterfaceActionType::Button, FInputChord());
+		UI_COMMAND(Command_TraceToggle, "Toggle Trace", "Toggles trace capture.", EUserInterfaceActionType::Button, FInputChord());
 	}
 	UE_ENABLE_OPTIMIZATION_SHIP
 
 	TSharedPtr<FUICommandInfo> Command_TraceScreenshot;
 	TSharedPtr<FUICommandInfo> Command_TraceBookmark;
+	TSharedPtr<FUICommandInfo> Command_TraceToggle;
 };
 
+////////////////////////////////////////////////////////////////////////////////////////////////////
+
 TSharedRef<SWidget> CreateInsightsStatusBarWidget()
 {
 	return SNew(SInsightsStatusBarWidget);
@@ -123,14 +128,14 @@ void SInsightsStatusBarWidget::Tick(const FGeometry& AllottedGeometry, const dou
 
 	if (bShouldTryOpenLiveSession)
 	{
-		constexpr double Delay = 1.0f;
+		constexpr double Delay = 1.0;
 		OpenLiveSessionScheduledTime = InCurrentTime + Delay;
 		bShouldTryOpenLiveSession = false;
 	}
 
-	if (OpenLiveSessionScheduledTime > 0.0f && OpenLiveSessionScheduledTime <= InCurrentTime)
+	if (OpenLiveSessionScheduledTime > 0.0 && OpenLiveSessionScheduledTime <= InCurrentTime)
 	{
-		OpenLiveSessionScheduledTime = 0.0f;
+		OpenLiveSessionScheduledTime = 0.0;
 		FUnrealInsightsLauncher::Get()->TryOpenTraceFromDestination(FTraceAuxiliary::GetTraceDestinationString());
 	}
 }
@@ -157,7 +162,7 @@ void SInsightsStatusBarWidget::Construct(const FArguments& InArgs)
 				SNew(SHorizontalBox)
 
 				+ SHorizontalBox::Slot()
-				.Padding(0, 0, 3, 0)
+				.Padding(0.0f, 0.0f, 3.0f, 0.0f)
 				.AutoWidth()
 				[
 					SNew(SImage)
@@ -181,7 +186,7 @@ void SInsightsStatusBarWidget::Construct(const FArguments& InArgs)
 		[
 			SNew(UE::TraceTools::SToggleTraceButton)
 			.OnToggleTraceRequested(this, &SInsightsStatusBarWidget::ToggleTrace_OnClicked)
-			.IsTraceRunning_Lambda([]() {return UE::Trace::IsTracing(); })
+			.IsTraceRunning_Lambda([]() { return UE::Trace::IsTracing(); })
 			.ButtonSize(UE::TraceTools::SToggleTraceButton::EButtonSize::StatusBar)
 			.IsEnabled(this, &SInsightsStatusBarWidget::ToggleTrace_CanExecute)
 		]
@@ -270,18 +275,19 @@ TSharedRef<SWidget> SInsightsStatusBarWidget::MakeTraceMenu()
 			TAttribute<FText>::CreateSP(this, &SInsightsStatusBarWidget::GetTraceBookmarkTooltipText),
 			FSlateIcon(FEditorTraceUtilitiesStyle::Get().GetStyleSetName(), "Icons.Bookmark.Menu")
 		);
-		
+
 		MenuBuilder.AddEditableText(LOCTEXT("TraceRegionNameLabel", "Region Name:"),
-							GetTraceRegionNameDesc(),
-									FSlateIcon(),
-									TAttribute<FText>::Create(TAttribute<FText>::FGetter::CreateRaw(this, &SInsightsStatusBarWidget::GetTraceRegionName)),
-									FOnTextCommitted::CreateLambda([](const FText& NewRegionName, ETextCommit::Type)
-									{
-										GConfig->SetText(SettingsCategory, TraceRegionSettingName, NewRegionName, FEditorTraceUtilitiesModule::GetTraceUtilitiesIni());
-									}),
-									FOnTextChanged(),
-									RegionIsActive());
-		
+			GetTraceRegionNameDesc(),
+			FSlateIcon(),
+			TAttribute<FText>::Create(TAttribute<FText>::FGetter::CreateRaw(this, &SInsightsStatusBarWidget::GetTraceRegionName)),
+			FOnTextCommitted::CreateLambda(
+				[](const FText& NewRegionName, ETextCommit::Type)
+				{
+					GConfig->SetText(SettingsCategory, TraceRegionSettingName, NewRegionName, FEditorTraceUtilitiesModule::GetTraceUtilitiesIni());
+				}),
+			FOnTextChanged(),
+			RegionIsActive());
+
 		MenuBuilder.AddMenuEntry(
 			GetRegionSwitchLabelText(),
 			TAttribute<FText>::CreateSP(this, &SInsightsStatusBarWidget::GetRegionSwitchDescText),
@@ -291,13 +297,13 @@ TSharedRef<SWidget> SInsightsStatusBarWidget::MakeTraceMenu()
 			NAME_None,
 			EUserInterfaceActionType::Button
 		);
-		
+
 		MenuBuilder.AddMenuEntry(
 			LOCTEXT("StatNamedEventsLabel", "Stat Named Events"),
 			LOCTEXT("StatNamedEventsDesc", "Enable or disable named events in the stats system."),
 			FSlateIcon(),
 			FUIAction(FExecuteAction::CreateLambda([]() { GCycleStatsShouldEmitNamedEvents = GCycleStatsShouldEmitNamedEvents == 0 ? 1 : 0; }),
-					  FCanExecuteAction::CreateLambda([]() { return true; }),
+				      FCanExecuteAction(),
 					  FIsActionChecked::CreateLambda([]() { return GCycleStatsShouldEmitNamedEvents > 0; })),
 			NAME_None,
 			EUserInterfaceActionType::ToggleButton
@@ -334,13 +340,11 @@ TSharedRef<SWidget> SInsightsStatusBarWidget::MakeTraceMenu()
 	MenuBuilder.BeginSection("Tracing", LOCTEXT("TraceMenu_Section_Tracing", "Tracing"));
 	{
 		MenuBuilder.AddMenuEntry(
+			FInsightsStatusBarWidgetCommands::Get().Command_TraceToggle,
+			NAME_None,
 			TAttribute<FText>::CreateSP(this, &SInsightsStatusBarWidget::GetTraceMenuItemText),
 			TAttribute<FText>::CreateSP(this, &SInsightsStatusBarWidget::GetTraceMenuItemTooltipText),
-			FSlateIcon(FEditorTraceUtilitiesStyle::Get().GetStyleSetName(), "Icons.StartTrace.Menu"),
-			FUIAction(FExecuteAction::CreateSP(this, &SInsightsStatusBarWidget::ToggleTrace_OnClicked),
-					  FCanExecuteAction::CreateSP(this, &SInsightsStatusBarWidget::ToggleTrace_CanExecute)),
-			NAME_None,
-			EUserInterfaceActionType::Button
+			FSlateIcon(FEditorTraceUtilitiesStyle::Get().GetStyleSetName(), "Icons.StartTrace.Menu")
 		);
 
 		if (FTraceAuxiliary::IsPaused())
@@ -383,12 +387,12 @@ TSharedRef<SWidget> SInsightsStatusBarWidget::MakeTraceMenu()
 	MenuBuilder.BeginSection("Options", LOCTEXT("TraceMenu_Section_Options", "Options"));
 	{
 		MenuBuilder.AddMenuEntry(
-			LOCTEXT("OpenLiveSesssionOnTraceStart", "Open Live Session on Trace Start"),
-			LOCTEXT("OpenLiveSesssionOnTraceStartDesc", "When set, the live session will be automatically opened in Unreal Insights when tracing is started.\nThis option will only apply when tracing to the trace store."),
+			LOCTEXT("OpenLiveSessionOnTraceStart", "Open Live Session on Trace Start"),
+			LOCTEXT("OpenLiveSessionOnTraceStartDesc", "When set, the live session will be automatically opened in Unreal Insights when tracing is started.\nThis option will only apply when tracing to the trace store."),
 			FSlateIcon(),
 			FUIAction(FExecuteAction::CreateSP(this, &SInsightsStatusBarWidget::ToggleBooleanSettingValue, OpenLiveSessionOnTraceStartSettingName),
-				FCanExecuteAction::CreateLambda([]() { return true; }),
-				FIsActionChecked::CreateSP(this, &SInsightsStatusBarWidget::GetBooleanSettingValue, OpenLiveSessionOnTraceStartSettingName)),
+					  FCanExecuteAction(),
+					  FIsActionChecked::CreateSP(this, &SInsightsStatusBarWidget::GetBooleanSettingValue, OpenLiveSessionOnTraceStartSettingName)),
 			NAME_None,
 			EUserInterfaceActionType::ToggleButton
 		);
@@ -398,8 +402,8 @@ TSharedRef<SWidget> SInsightsStatusBarWidget::MakeTraceMenu()
 			LOCTEXT("OpenInsightsAfterTraceDesc", "When set, the session will be automatically opened in Unreal Insights when tracing is stopped or when a snapshot is saved."),
 			FSlateIcon(),
 			FUIAction(FExecuteAction::CreateSP(this, &SInsightsStatusBarWidget::ToggleBooleanSettingValue, OpenInsightsAfterTraceSettingName),
-				FCanExecuteAction::CreateLambda([]() { return true; }),
-				FIsActionChecked::CreateSP(this, &SInsightsStatusBarWidget::GetBooleanSettingValue, OpenInsightsAfterTraceSettingName)),
+					  FCanExecuteAction(),
+					  FIsActionChecked::CreateSP(this, &SInsightsStatusBarWidget::GetBooleanSettingValue, OpenInsightsAfterTraceSettingName)),
 			NAME_None,
 			EUserInterfaceActionType::ToggleButton
 		);
@@ -409,8 +413,8 @@ TSharedRef<SWidget> SInsightsStatusBarWidget::MakeTraceMenu()
 			LOCTEXT("ShowInExplorerAfterTraceDesc", "When set, folder containing the recorded session will be opened automatically when trace is stopped or when a snapshot is saved."),
 			FSlateIcon(),
 			FUIAction(FExecuteAction::CreateSP(this, &SInsightsStatusBarWidget::ToggleBooleanSettingValue, ShowInExplorerAfterTraceSettingName),
-				FCanExecuteAction::CreateLambda([]() { return true; }),
-				FIsActionChecked::CreateSP(this, &SInsightsStatusBarWidget::GetBooleanSettingValue, ShowInExplorerAfterTraceSettingName)),
+					  FCanExecuteAction(),
+					  FIsActionChecked::CreateSP(this, &SInsightsStatusBarWidget::GetBooleanSettingValue, ShowInExplorerAfterTraceSettingName)),
 			NAME_None,
 			EUserInterfaceActionType::ToggleButton
 		);
@@ -444,7 +448,8 @@ TSharedRef<SWidget> SInsightsStatusBarWidget::MakeTraceMenu()
 		MenuBuilder.AddSubMenu(
 			LOCTEXT("ServerControlLabel", "Unreal Trace Server"),
 			LOCTEXT("ServerControlTooltip", "Info and controls for the Unreal Trace Server instances"),
-			FNewMenuDelegate::CreateLambda([this](FMenuBuilder& MenuBuilder)
+			FNewMenuDelegate::CreateLambda(
+				[this](FMenuBuilder& MenuBuilder)
 				{
 					for (auto& ServerControl : ServerControls)
 					{
@@ -469,7 +474,7 @@ TSharedRef<SWidget> SInsightsStatusBarWidget::MakeTraceMenu()
 			LOCTEXT("OpenLiveSessionTooltip", "Opening the live session is possible only while tracing to the trace store."),
 			FSlateIcon(FEditorTraceUtilitiesStyle::Get().GetStyleSetName(), "Icons.OpenLiveSession.Menu"),
 			FUIAction(FExecuteAction::CreateSP(this, &SInsightsStatusBarWidget::OpenLiveSession_OnClicked),
-				FCanExecuteAction::CreateLambda([this]() { return FTraceAuxiliary::GetConnectionType() == FTraceAuxiliary::EConnectionType::Network; })),
+					  FCanExecuteAction::CreateLambda([this]() { return FTraceAuxiliary::GetConnectionType() == FTraceAuxiliary::EConnectionType::Network; })),
 			NAME_None,
 			EUserInterfaceActionType::Button
 		);
@@ -477,7 +482,7 @@ TSharedRef<SWidget> SInsightsStatusBarWidget::MakeTraceMenu()
 		MenuBuilder.AddSubMenu
 		(
 			LOCTEXT("RecentTracesLabel", "Recent Traces"),
-			LOCTEXT("RecentTracesTooltop", "Open the latest traces recorded to the trace store or as a file."),
+			LOCTEXT("RecentTracesTooltip", "Open the latest traces recorded to the trace store or as a file."),
 			FNewMenuDelegate::CreateSP(this, &SInsightsStatusBarWidget::Traces_BuildMenu),
 			false,
 			FSlateIcon(FEditorTraceUtilitiesStyle::Get().GetStyleSetName(), ("Icons.Trace.Menu"))
@@ -497,19 +502,15 @@ void SInsightsStatusBarWidget::Channels_BuildMenu(FMenuBuilder& MenuBuilder)
 		for (int32 Index = 0; Index < ChannelsInfo.Num(); ++Index)
 		{
 			const FChannelData& Data = ChannelsInfo[Index];
+
 			FString ChannelDisplayName = Data.Name;
 			ChannelDisplayName.RemoveFromEnd(TEXT("Channel"), 7);
+
 			FText ChannelTooltip = FText::FromString(Data.Desc);
 			if (Data.bIsReadOnly)
 			{
-				if (Data.Desc.EndsWith("."))
-				{
-					ChannelTooltip = FText::Format(LOCTEXT("ChannelTooltipFmt1", "{0} This channel is readonly and can only be enabled from the command line."), FText::FromString(Data.Desc));
-				}
-				else
-				{
-					ChannelTooltip = FText::Format(LOCTEXT("ChannelTooltipFmt2", "{0}. This channel is readonly and can only be enabled from the command line."), FText::FromString(Data.Desc));
-				}
+				ChannelTooltip = FText::Format(LOCTEXT("ChannelTooltipFmt", "{0}{1} This channel is readonly and can only be enabled from the command line."),
+					ChannelTooltip, Data.Desc.EndsWith(".") ? FText() : FText::FromString("."));
 			}
 
 			MenuBuilder.AddMenuEntry(
@@ -517,8 +518,8 @@ void SInsightsStatusBarWidget::Channels_BuildMenu(FMenuBuilder& MenuBuilder)
 				ChannelTooltip,
 				FSlateIcon(),
 				FUIAction(FExecuteAction::CreateSP(this, &SInsightsStatusBarWidget::ToggleChannel_Execute, Index),
-					FCanExecuteAction::CreateLambda([Value = !Data.bIsReadOnly]() { return Value; }),
-					FIsActionChecked::CreateSP(this, &SInsightsStatusBarWidget::ToggleChannel_IsChecked, Index)),
+						  FCanExecuteAction::CreateLambda([Value = !Data.bIsReadOnly]() { return Value; }),
+						  FIsActionChecked::CreateSP(this, &SInsightsStatusBarWidget::ToggleChannel_IsChecked, Index)),
 				NAME_None,
 				EUserInterfaceActionType::ToggleButton);
 		}
@@ -546,14 +547,29 @@ void SInsightsStatusBarWidget::Traces_BuildMenu(FMenuBuilder& MenuBuilder)
 void SInsightsStatusBarWidget::InitCommandList()
 {
 	FInsightsStatusBarWidgetCommands::Register();
+
 	CommandList = MakeShared<FUICommandList>();
-	CommandList->MapAction(FInsightsStatusBarWidgetCommands::Get().Command_TraceScreenshot, FExecuteAction::CreateSP(this, &SInsightsStatusBarWidget::TraceScreenshot_Execute), FCanExecuteAction::CreateSP(this, &SInsightsStatusBarWidget::TraceScreenshot_CanExecute));
-	CommandList->MapAction(FInsightsStatusBarWidgetCommands::Get().Command_TraceBookmark, FExecuteAction::CreateSP(this, &SInsightsStatusBarWidget::TraceBookmark_Execute), FCanExecuteAction::CreateSP(this, &SInsightsStatusBarWidget::TraceBookmark_CanExecute));
+
+	CommandList->MapAction(FInsightsStatusBarWidgetCommands::Get().Command_TraceScreenshot,
+		FExecuteAction::CreateSP(this, &SInsightsStatusBarWidget::TraceScreenshot_Execute),
+		FCanExecuteAction::CreateSP(this, &SInsightsStatusBarWidget::TraceScreenshot_CanExecute));
+
+	CommandList->MapAction(FInsightsStatusBarWidgetCommands::Get().Command_TraceBookmark,
+		FExecuteAction::CreateSP(this, &SInsightsStatusBarWidget::TraceBookmark_Execute),
+		FCanExecuteAction::CreateSP(this, &SInsightsStatusBarWidget::TraceBookmark_CanExecute));
+
+	CommandList->MapAction(FInsightsStatusBarWidgetCommands::Get().Command_TraceToggle,
+		FExecuteAction::CreateSP(this, &SInsightsStatusBarWidget::ToggleTrace_OnClicked),
+		FCanExecuteAction::CreateSP(this, &SInsightsStatusBarWidget::ToggleTrace_CanExecute));
+
+	// Ensures the actions are globally accessible.
+	IMainFrameModule& MainFrame = FModuleManager::Get().LoadModuleChecked<IMainFrameModule>("MainFrame");
+	MainFrame.GetMainFrameCommandBindings()->Append(CommandList.ToSharedRef());
 }
 
 FText SInsightsStatusBarWidget::GetTitleToolTipText() const
 {
-#if UE_TRACE_ENABLED
+#if UE_TRACE_MINIMAL_ENABLED
 	FTextBuilder DescBuilder;
 
 	const FString Dest = FTraceAuxiliary::GetTraceDestinationString();
@@ -573,7 +589,7 @@ FText SInsightsStatusBarWidget::GetTitleToolTipText() const
 
 	return DescBuilder.ToText();
 #else
-	return LOCTEXT("TraceStatusDisabled", "Trace system is disabled at compile time. Check the UE_TRACE_ENABLED define.");
+	return LOCTEXT("TraceStatusDisabled", "Trace system is disabled at compile time. Check the UE_TRACE_MINIMAL_ENABLED define.");
 #endif
 }
 
@@ -704,12 +720,7 @@ bool SInsightsStatusBarWidget::SetTraceDestination_IsChecked(ETraceDestination I
 
 bool SInsightsStatusBarWidget::SetTraceDestination_CanExecute()
 {
-	if (!UE::Trace::IsTracing())
-	{
-		return true;
-	}
-
-	return false;
+	return !UE::Trace::IsTracing();
 }
 
 void SInsightsStatusBarWidget::SaveSnapshot()
@@ -737,7 +748,7 @@ void SInsightsStatusBarWidget::SaveSnapshot()
 
 bool SInsightsStatusBarWidget::SaveSnapshot_CanExecute() const
 {
-#if UE_TRACE_ENABLED
+#if UE_TRACE_MINIMAL_ENABLED
 	return true;
 #else
 	return false;
@@ -756,7 +767,7 @@ FText SInsightsStatusBarWidget::GetTraceMenuItemText() const
 
 FText SInsightsStatusBarWidget::GetTraceMenuItemTooltipText() const
 {
-#if UE_TRACE_ENABLED
+#if UE_TRACE_MINIMAL_ENABLED
 	if (UE::Trace::IsTracing())
 	{
 		return LOCTEXT("StopTraceButtonTooltip", "Stop tracing");
@@ -764,13 +775,13 @@ FText SInsightsStatusBarWidget::GetTraceMenuItemTooltipText() const
 
 	return LOCTEXT("StartTraceButtonTooltip", "Start tracing to the selected trace destination.");
 #else
-	return LOCTEXT("StartTraceDisabledButtonTooltip", "Trace system is disabled at compile time. Check the UE_TRACE_ENABLED define.");
+	return LOCTEXT("StartTraceDisabledButtonTooltip", "Trace system is disabled at compile time. Check the UE_TRACE_MINIMAL_ENABLED define.");
 #endif
 }
 
 bool SInsightsStatusBarWidget::ToggleTrace_CanExecute() const
 {
-#if UE_TRACE_ENABLED
+#if UE_TRACE_MINIMAL_ENABLED
 	return true;
 #else
 	return false;
@@ -804,7 +815,7 @@ void SInsightsStatusBarWidget::ToggleTrace_OnClicked()
 				TraceDestinationStr = TEXT("External Target");
 			}
 
-			UE_LOG(LogTraceUtilities, Log, TEXT("Trace started. Trace is now active and saving to the following location (file or tracestore): %s"), *TraceDestinationStr);
+			UE_LOG(LogTraceUtilities, Log, TEXT("Trace started. Trace is now active and saving to the following location (file or trace store): %s"), *TraceDestinationStr);
 		}
 		else
 		{
@@ -829,7 +840,7 @@ FText SInsightsStatusBarWidget::GetPauseTraceMenuItemTooltipText() const
 }
 
 void SInsightsStatusBarWidget::TogglePauseTrace_OnClicked()
-{ 
+{
 	if (!UE::Trace::IsTracing())
 	{
 		return;
@@ -952,7 +963,7 @@ void SInsightsStatusBarWidget::CacheTraceStorePath()
 
 void SInsightsStatusBarWidget::CreateChannelsInfo()
 {
-#if UE_TRACE_ENABLED
+#if UE_TRACE_MINIMAL_ENABLED
 	ChannelsInfo.Empty();
 
 	UE::Trace::EnumerateChannels([](const UE::Trace::FChannelInfo& Info, void* User)
@@ -969,7 +980,12 @@ void SInsightsStatusBarWidget::CreateChannelsInfo()
 		return true;
 	}, &ChannelsInfo);
 
-	Algo::SortBy(ChannelsInfo, [](const FChannelData& Entry) { return Entry.Name; }, [](const FString& Rhs, const FString& Lhs)
+	Algo::SortBy(ChannelsInfo,
+		[](const FChannelData& Entry)
+		{
+			return Entry.Name;
+		},
+		[](const FString& Rhs, const FString& Lhs)
 		{
 			return Rhs.Compare(Lhs) < 0;
 		});
@@ -978,13 +994,18 @@ void SInsightsStatusBarWidget::CreateChannelsInfo()
 
 void SInsightsStatusBarWidget::UpdateChannelsInfo()
 {
-#if UE_TRACE_ENABLED
+#if UE_TRACE_MINIMAL_ENABLED
 	UE::Trace::EnumerateChannels([](const UE::Trace::FChannelInfo& Info, void* User)
 	{
 		TArray<FChannelData>* Channels = (TArray<FChannelData>*) User;
 
 		FString Name = Info.Name;
-		int32 Index = Algo::BinarySearchBy(*Channels, Name, [](const FChannelData& Entry) { return Entry.Name; }, [](const FString& Rhs, const FString& Lhs)
+		int32 Index = Algo::BinarySearchBy(*Channels, Name,
+			[](const FChannelData& Entry)
+			{
+				return Entry.Name;
+			},
+			[](const FString& Rhs, const FString& Lhs)
 			{
 				return Rhs.Compare(Lhs) < 0;
 			});
@@ -992,9 +1013,7 @@ void SInsightsStatusBarWidget::UpdateChannelsInfo()
 		if (Index != INDEX_NONE)
 		{
 			FChannelData& Data = (*Channels)[Index];
-			Data.Name = Info.Name;
 			Data.bIsEnabled = Info.bIsEnabled;
-			Data.bIsReadOnly = Info.bIsReadOnly;
 		}
 
 		return true;
@@ -1029,7 +1048,6 @@ bool SInsightsStatusBarWidget::ToggleChannel_IsChecked(int32 Index)
 	return false;
 }
 
-
 bool SInsightsStatusBarWidget::TraceScreenshot_CanExecute() const
 {
 	return SHOULD_TRACE_SCREENSHOT();
@@ -1060,12 +1078,14 @@ FText SInsightsStatusBarWidget::GetTraceRegionName()
 
 FText SInsightsStatusBarWidget::GetTraceRegionNameDesc()
 {
-	return RegionIsActive() ? LOCTEXT("TraceRegionNameDisabledDesc", "The name of the region to start or stop. Cannot be edited while a region is active.") : LOCTEXT("TraceRegionNameEnabledDesc", "The name of the region to start or stop.");
+	return RegionIsActive() ?
+		LOCTEXT("TraceRegionNameDesc_Active", "The name of the region to start or stop. Cannot be edited while a region is active.") :
+		LOCTEXT("TraceRegionNameDesc_Inactive", "The name of the region to start or stop.");
 }
 
 void SInsightsStatusBarWidget::ToggleRegion_Execute()
 {
-	if(RegionIsActive())
+	if (RegionIsActive())
 	{
 		TRACE_END_REGION_WITH_ID(RegionId);
 		RegionId = 0;
@@ -1088,9 +1108,9 @@ bool SInsightsStatusBarWidget::RegionIsActive() const
 
 FText SInsightsStatusBarWidget::GetRegionSwitchLabelText() const
 {
-	if(RegionIsActive())
+	if (RegionIsActive())
 	{
-		return LOCTEXT("TraceEndRegionLabel", "End Region");	
+		return LOCTEXT("TraceEndRegionLabel", "End Region");
 	}
 	return LOCTEXT("TraceBeginRegionLabel", "Begin Region");
 }
@@ -1102,9 +1122,9 @@ FText SInsightsStatusBarWidget::GetRegionSwitchDescText() const
 		return LOCTEXT("TraceRegionDisabledDesc", "Regions can only be traced when the Region channel in enabled.");
 	}
 
-	if(RegionIsActive())
+	if (RegionIsActive())
 	{
-		return LOCTEXT("TraceEndRegionDesc", "Marks the ending of a trace region with the name input above.");	
+		return LOCTEXT("TraceEndRegionDesc", "Marks the ending of a trace region with the name input above.");
 	}
 
 	return LOCTEXT("TraceBeginRegionDesc", "Marks the beginning of a trace region with the name input above.");
@@ -1134,12 +1154,12 @@ void SInsightsStatusBarWidget::PopulateRecentTracesList()
 	bIsFromTraceStore = false;
 	IFileManager::Get().IterateDirectoryStat(*FPaths::ProfilingDir(), Visitor);
 
-	Algo::SortBy(Traces, [](TSharedPtr<FTraceFileInfo> TraceInfo) { return *TraceInfo; });
+	Algo::SortBy(Traces, [](const TSharedPtr<FTraceFileInfo>& TraceInfo) { return *TraceInfo; });
 
 	constexpr int32 MaxRecentTraces = 15;
 	if (Traces.Num() > MaxRecentTraces)
 	{
-		Traces.RemoveAt(MaxRecentTraces - 1, Traces.Num() - MaxRecentTraces);
+		Traces.RemoveAt(MaxRecentTraces, Traces.Num() - MaxRecentTraces);
 	}
 }
 
